Copying //UE4/Dev-Framework to //UE4/Dev-Main (Source: //UE4/Dev-Framework @ 3136612)

diff --git a/Engine/Source/Developer/SimplygonSwarm/Private/SimplygonSwarm.cpp b/Engine/Source/Developer/SimplygonSwarm/Private/SimplygonSwarm.cpp
--- a/Engine/Source/Developer/SimplygonSwarm/Private/SimplygonSwarm.cpp
+++ b/Engine/Source/Developer/SimplygonSwarm/Private/SimplygonSwarm.cpp
@@ -884,33 +884,47 @@ private:
 
 			if (ChannelName.Compare(BASECOLOR_CHANNEL) == 0)
 			{
-				ExtractTextureDescriptors(SceneGraph, Channel, OutputFolderPath, ChannelName, OutMaterial.DiffuseSamples, OutMaterial.DiffuseSize);
+				FIntPoint Size = OutMaterial.GetPropertySize(EFlattenMaterialProperties::Diffuse);
+				ExtractTextureDescriptors(SceneGraph, Channel, OutputFolderPath, ChannelName, OutMaterial.GetPropertySamples(EFlattenMaterialProperties::Diffuse), Size);
+				OutMaterial.SetPropertySize(EFlattenMaterialProperties::Diffuse, Size);
 			}
 			else if (ChannelName.Compare(NORMAL_CHANNEL) == 0)
 			{
-				ExtractTextureDescriptors(SceneGraph, Channel, OutputFolderPath, ChannelName, OutMaterial.NormalSamples, OutMaterial.NormalSize);
+				FIntPoint Size = OutMaterial.GetPropertySize(EFlattenMaterialProperties::Normal);
+				ExtractTextureDescriptors(SceneGraph, Channel, OutputFolderPath, ChannelName, OutMaterial.GetPropertySamples(EFlattenMaterialProperties::Normal), Size);
+				OutMaterial.SetPropertySize(EFlattenMaterialProperties::Normal, Size);
 			}
 			else if (ChannelName.Compare(SPECULAR_CHANNEL) == 0)
 			{
-				ExtractTextureDescriptors(SceneGraph, Channel, OutputFolderPath, ChannelName, OutMaterial.SpecularSamples, OutMaterial.SpecularSize);
+				FIntPoint Size = OutMaterial.GetPropertySize(EFlattenMaterialProperties::Specular);
+				ExtractTextureDescriptors(SceneGraph, Channel, OutputFolderPath, ChannelName, OutMaterial.GetPropertySamples(EFlattenMaterialProperties::Specular),	Size);
+				OutMaterial.SetPropertySize(EFlattenMaterialProperties::Specular, Size);
 			}
 			else if (ChannelName.Compare(ROUGHNESS_CHANNEL) == 0)
 			{
-				ExtractTextureDescriptors(SceneGraph, Channel, OutputFolderPath, ChannelName, OutMaterial.RoughnessSamples, OutMaterial.RoughnessSize);
+				FIntPoint Size = OutMaterial.GetPropertySize(EFlattenMaterialProperties::Roughness);
+				ExtractTextureDescriptors(SceneGraph, Channel, OutputFolderPath, ChannelName, OutMaterial.GetPropertySamples(EFlattenMaterialProperties::Roughness),Size);
+				OutMaterial.SetPropertySize(EFlattenMaterialProperties::Roughness, Size);
 			}
 			else if (ChannelName.Compare(METALLIC_CHANNEL) == 0)
 			{
-				ExtractTextureDescriptors(SceneGraph, Channel, OutputFolderPath, ChannelName, OutMaterial.MetallicSamples, OutMaterial.MetallicSize);
+				FIntPoint Size = OutMaterial.GetPropertySize(EFlattenMaterialProperties::Metallic);
+				ExtractTextureDescriptors(SceneGraph, Channel, OutputFolderPath, ChannelName, OutMaterial.GetPropertySamples(EFlattenMaterialProperties::Metallic), Size);
+				OutMaterial.SetPropertySize(EFlattenMaterialProperties::Metallic, Size);
 			}
 			else if (ChannelName.Compare(OPACITY_CHANNEL) == 0)
 			{
-				ExtractTextureDescriptors(SceneGraph, Channel, OutputFolderPath, ChannelName, OutMaterial.OpacitySamples, OutMaterial.OpacitySize);
+				FIntPoint Size = OutMaterial.GetPropertySize(EFlattenMaterialProperties::Opacity);
+				ExtractTextureDescriptors(SceneGraph, Channel, OutputFolderPath, ChannelName, OutMaterial.GetPropertySamples(EFlattenMaterialProperties::Opacity), Size);
 				bHasOpacity = true;
+				OutMaterial.SetPropertySize(EFlattenMaterialProperties::Opacity, Size);
 			}
 			else if (ChannelName.Compare(OPACITY_MASK_CHANNEL) == 0)
 			{
-				ExtractTextureDescriptors(SceneGraph, Channel, OutputFolderPath, ChannelName, OutMaterial.OpacitySamples, OutMaterial.OpacitySize);
+				FIntPoint Size = OutMaterial.GetPropertySize(EFlattenMaterialProperties::Opacity);
+				ExtractTextureDescriptors(SceneGraph, Channel, OutputFolderPath, ChannelName, OutMaterial.GetPropertySamples(EFlattenMaterialProperties::Opacity), Size);
 				bHasOpacityMask = true;
+				OutMaterial.SetPropertySize(EFlattenMaterialProperties::Opacity, Size);
 			}
 		}
 
@@ -1654,10 +1668,10 @@ bool ZipContentsForUpload(FString InputDirectoryPath, FString OutputFileName)
 
 
 		// Does current material have BaseColor?
-		if (FlattenMaterial.DiffuseSamples.Num())
+		if (FlattenMaterial.DoesPropertyContainData(EFlattenMaterialProperties::Diffuse))
 		{
 			FString ChannelName(BASECOLOR_CHANNEL);
-			ssf::pssfMaterialChannel BaseColorChannel = SetMaterialChannelData(FlattenMaterial.DiffuseSamples, FlattenMaterial.DiffuseSize, SgTextureTable, ChannelName, FString::Printf(TEXT("%s%s"), *MaterialName, *ChannelName), OutputFolderPath);
+			ssf::pssfMaterialChannel BaseColorChannel = SetMaterialChannelData(FlattenMaterial.GetPropertySamples(EFlattenMaterialProperties::Diffuse), FlattenMaterial.GetPropertySize(EFlattenMaterialProperties::Diffuse), SgTextureTable, ChannelName, FString::Printf(TEXT("%s%s"), *MaterialName, *ChannelName), OutputFolderPath);
 				 
 		// if (InMaterialLODSettings.ChannelsToCast[0].bBakeVertexColors)
 			//{
@@ -1669,51 +1683,51 @@ bool ZipContentsForUpload(FString InputDirectoryPath, FString OutputFileName)
 		}
 
 		// Does current material have Metallic?
-		if (FlattenMaterial.MetallicSamples.Num())
+		if (FlattenMaterial.DoesPropertyContainData(EFlattenMaterialProperties::Metallic))
 		{
 			FString ChannelName(METALLIC_CHANNEL);
-			ssf::pssfMaterialChannel MetallicChannel = SetMaterialChannelData(FlattenMaterial.MetallicSamples, FlattenMaterial.MetallicSize, SgTextureTable, ChannelName, FString::Printf(TEXT("%s%s"), *MaterialName, *ChannelName), OutputFolderPath);
+			ssf::pssfMaterialChannel MetallicChannel = SetMaterialChannelData(FlattenMaterial.GetPropertySamples(EFlattenMaterialProperties::Metallic), FlattenMaterial.GetPropertySize(EFlattenMaterialProperties::Metallic), SgTextureTable, ChannelName, FString::Printf(TEXT("%s%s"), *MaterialName, *ChannelName), OutputFolderPath);
 			SgMaterial->MaterialChannelList.push_back(MetallicChannel);
 		}
 
 		// Does current material have Specular?
-		if (FlattenMaterial.SpecularSamples.Num())
+		if (FlattenMaterial.DoesPropertyContainData(EFlattenMaterialProperties::Specular))
 		{
 			FString ChannelName(SPECULAR_CHANNEL);
-			ssf::pssfMaterialChannel SpecularChannel = SetMaterialChannelData(FlattenMaterial.SpecularSamples, FlattenMaterial.SpecularSize, SgTextureTable, ChannelName, FString::Printf(TEXT("%s%s"), *MaterialName, *ChannelName), OutputFolderPath);
+			ssf::pssfMaterialChannel SpecularChannel = SetMaterialChannelData(FlattenMaterial.GetPropertySamples(EFlattenMaterialProperties::Specular), FlattenMaterial.GetPropertySize(EFlattenMaterialProperties::Specular), SgTextureTable, ChannelName, FString::Printf(TEXT("%s%s"), *MaterialName, *ChannelName), OutputFolderPath);
 			SgMaterial->MaterialChannelList.push_back(SpecularChannel);
 		}
 
 		// Does current material have Roughness?
-		if (FlattenMaterial.RoughnessSamples.Num())
+		if (FlattenMaterial.DoesPropertyContainData(EFlattenMaterialProperties::Roughness))
 		{
 			FString ChannelName(ROUGHNESS_CHANNEL);
-			ssf::pssfMaterialChannel RoughnessChannel = SetMaterialChannelData(FlattenMaterial.RoughnessSamples, FlattenMaterial.RoughnessSize, SgTextureTable, ChannelName, FString::Printf(TEXT("%s%s"), *MaterialName, *ChannelName), OutputFolderPath);
+			ssf::pssfMaterialChannel RoughnessChannel = SetMaterialChannelData(FlattenMaterial.GetPropertySamples(EFlattenMaterialProperties::Roughness), FlattenMaterial.GetPropertySize(EFlattenMaterialProperties::Roughness), SgTextureTable, ChannelName, FString::Printf(TEXT("%s%s"), *MaterialName, *ChannelName), OutputFolderPath);
 			SgMaterial->MaterialChannelList.push_back(RoughnessChannel);
 		}
 
 		//Does current material have a normalmap?
-		if (FlattenMaterial.NormalSamples.Num())
+		if (FlattenMaterial.DoesPropertyContainData(EFlattenMaterialProperties::Normal))
 		{
 			FString ChannelName(NORMAL_CHANNEL);
 			SgMaterial->TangentSpaceNormals.Create();
 			SgMaterial->TangentSpaceNormals.Set(true);
-			ssf::pssfMaterialChannel NormalChannel = SetMaterialChannelData(FlattenMaterial.NormalSamples, FlattenMaterial.NormalSize, SgTextureTable, ChannelName, FString::Printf(TEXT("%s%s"), *MaterialName, *ChannelName), OutputFolderPath, false);
+			ssf::pssfMaterialChannel NormalChannel = SetMaterialChannelData(FlattenMaterial.GetPropertySamples(EFlattenMaterialProperties::Normal), FlattenMaterial.GetPropertySize(EFlattenMaterialProperties::Normal), SgTextureTable, ChannelName, FString::Printf(TEXT("%s%s"), *MaterialName, *ChannelName), OutputFolderPath, false);
 			SgMaterial->MaterialChannelList.push_back(NormalChannel);
 		}
 
 		// Does current material have Opacity?
-		if (FlattenMaterial.OpacitySamples.Num())
+		if (FlattenMaterial.DoesPropertyContainData(EFlattenMaterialProperties::Opacity))
 		{
 			FString ChannelName(OPACITY_CHANNEL);
-			ssf::pssfMaterialChannel OpacityChannel = SetMaterialChannelData(FlattenMaterial.OpacitySamples, FlattenMaterial.OpacitySize, SgTextureTable, ChannelName, FString::Printf(TEXT("%s%s"), *MaterialName, *ChannelName), OutputFolderPath);
+			ssf::pssfMaterialChannel OpacityChannel = SetMaterialChannelData(FlattenMaterial.GetPropertySamples(EFlattenMaterialProperties::Opacity), FlattenMaterial.GetPropertySize(EFlattenMaterialProperties::Opacity), SgTextureTable, ChannelName, FString::Printf(TEXT("%s%s"), *MaterialName, *ChannelName), OutputFolderPath);
 			SgMaterial->MaterialChannelList.push_back(OpacityChannel);
 		}
 
-		if (FlattenMaterial.EmissiveSamples.Num())
+		if (FlattenMaterial.DoesPropertyContainData(EFlattenMaterialProperties::Emissive))
 		{
 			FString ChannelName(EMISSIVE_CHANNEL);
-			ssf::pssfMaterialChannel EmissiveChannel = SetMaterialChannelData(FlattenMaterial.EmissiveSamples, FlattenMaterial.EmissiveSize, SgTextureTable, ChannelName, FString::Printf(TEXT("%s%s"), *MaterialName, *ChannelName), OutputFolderPath);
+			ssf::pssfMaterialChannel EmissiveChannel = SetMaterialChannelData(FlattenMaterial.GetPropertySamples(EFlattenMaterialProperties::Emissive), FlattenMaterial.GetPropertySize(EFlattenMaterialProperties::Emissive), SgTextureTable, ChannelName, FString::Printf(TEXT("%s%s"), *MaterialName, *ChannelName), OutputFolderPath);
 			SgMaterial->MaterialChannelList.push_back(EmissiveChannel);
 		}
 		else
@@ -1721,7 +1735,7 @@ bool ZipContentsForUpload(FString InputDirectoryPath, FString OutputFileName)
 			FString ChannelName(EMISSIVE_CHANNEL);
 			TArray<FColor> BlackEmissive;
 			BlackEmissive.AddZeroed(1);
-			ssf::pssfMaterialChannel EmissiveChannel = SetMaterialChannelData(FlattenMaterial.EmissiveSamples, FlattenMaterial.EmissiveSize, SgTextureTable, ChannelName, FString::Printf(TEXT("%s%s"), *MaterialName, *ChannelName), OutputFolderPath);
+			ssf::pssfMaterialChannel EmissiveChannel = SetMaterialChannelData(FlattenMaterial.GetPropertySamples(EFlattenMaterialProperties::Emissive), FlattenMaterial.GetPropertySize(EFlattenMaterialProperties::Emissive), SgTextureTable, ChannelName, FString::Printf(TEXT("%s%s"), *MaterialName, *ChannelName), OutputFolderPath);
 			SgMaterial->MaterialChannelList.push_back(EmissiveChannel);
 		}
 			
diff --git a/Engine/Source/Editor/AnimGraph/AnimGraph.Build.cs b/Engine/Source/Editor/AnimGraph/AnimGraph.Build.cs
--- a/Engine/Source/Editor/AnimGraph/AnimGraph.Build.cs
+++ b/Engine/Source/Editor/AnimGraph/AnimGraph.Build.cs
@@ -28,15 +28,22 @@ public class AnimGraph : ModuleRules
                 "GraphEditor",
 				"PropertyEditor",
 				"EditorStyle",
-                "ContentBrowser"
+                "ContentBrowser",
 			}
 		);
 
-		CircularlyReferencedDependentModules.AddRange(
+        CircularlyReferencedDependentModules.AddRange(
 			new string[] {
                 "UnrealEd",
                 "GraphEditor",
             }
-		); 
-	}
+		);
+
+        PrivateIncludePathModuleNames.AddRange(
+            new string[] {
+                "Persona",
+                "SkeletonEditor"
+            }
+        );
+    }
 }
diff --git a/Engine/Source/Editor/AnimGraph/Classes/AnimGraphNode_Base.h b/Engine/Source/Editor/AnimGraph/Classes/AnimGraphNode_Base.h
--- a/Engine/Source/Editor/AnimGraph/Classes/AnimGraphNode_Base.h
+++ b/Engine/Source/Editor/AnimGraph/Classes/AnimGraphNode_Base.h
@@ -129,6 +129,7 @@ class ANIMGRAPH_API UAnimGraphNode_Base : public UK2Node
 	virtual FString GetDocumentationLink() const override;
 	virtual void GetPinHoverText(const UEdGraphPin& Pin, FString& HoverTextOut) const override;
 	virtual bool ShowPaletteIconOnNode() const override{ return false; }
+	virtual void PinDefaultValueChanged(UEdGraphPin* Pin) override;
 	// End of UEdGraphNode interface
 
 	// UK2Node interface
@@ -190,6 +191,35 @@ class ANIMGRAPH_API UAnimGraphNode_Base : public UK2Node
 	template<class AssetType>
 	void HandleAnimReferenceReplacement(AssetType*& OriginalAsset, const TMap<UAnimationAsset*, UAnimationAsset*>& AnimAssetReplacementMap);
 
+	/**
+	 * Selection notification callback.
+	 * If a node needs to handle viewport input etc. then it should push an editor mode here.
+	 * @param	bInIsSelected	Whether we selected or deselected the node
+	 * @param	InModeTools		The mode tools. Use this to push the editor mode if required.
+	 * @param	InRuntimeNode	The runtime node to go with this skeletal control. This may be NULL in some cases when bInIsSelected is false.
+	 */
+	virtual void OnNodeSelected(bool bInIsSelected, class FEditorModeTools& InModeTools, struct FAnimNode_Base* InRuntimeNode);
+
+	/**
+	 * Override this function to push an editor mode when this node is selected
+	 * @return the editor mode to use when this node is selected
+	 */
+	virtual FEditorModeID GetEditorMode() const;
+
+	// Draw function for supporting visualization
+	virtual void Draw(FPrimitiveDrawInterface* PDI, USkeletalMeshComponent * PreviewSkelMeshComp) const {}
+	// Canvas draw function to draw to viewport
+	virtual void DrawCanvas(FViewport& InViewport, FSceneView& View, FCanvas& Canvas, USkeletalMeshComponent * PreviewSkelMeshComp) const {}
+	// Function to collect strings from nodes to display in the viewport.
+	// Use this rather than DrawCanvas when adding general text to the viewport.
+	virtual void GetOnScreenDebugInfo(TArray<FText>& DebugInfo, USkeletalMeshComponent* PreviewSkelMeshComp) const {}
+
+	/** Called after editing a default value to update internal node from pin defaults. This is needed for forwarding code to propagate values to preview. */
+	virtual void CopyPinDefaultsToNodeData(UEdGraphPin* InPin) {}
+
+	/** Called to propagate data from the internal node to the preview in Persona. */
+	virtual void CopyNodeDataToPreviewNode(FAnimNode_Base* InPreviewNode) {}
+
 	// BEGIN Interface to support transition getter
 	// if you return true for DoesSupportExposeTimeForTransitionGetter
 	// you should implement all below functions
@@ -202,6 +232,9 @@ class ANIMGRAPH_API UAnimGraphNode_Base : public UK2Node
 	// can customize details tab 
 	virtual void CustomizeDetails(IDetailLayoutBuilder& DetailBuilder){ }
 
+	/** Try to find the preview node instance for this anim graph node */
+	FAnimNode_Base* FindDebugAnimNode(USkeletalMeshComponent * PreviewSkelMeshComp) const;
+
 	template<typename NodeType>
 	NodeType* GetActiveInstanceNode(UObject* AnimInstanceObject) const
 	{
diff --git a/Engine/Source/Editor/AnimGraph/Classes/AnimGraphNode_Fabrik.h b/Engine/Source/Editor/AnimGraph/Classes/AnimGraphNode_Fabrik.h
--- a/Engine/Source/Editor/AnimGraph/Classes/AnimGraphNode_Fabrik.h
+++ b/Engine/Source/Editor/AnimGraph/Classes/AnimGraphNode_Fabrik.h
@@ -19,16 +19,11 @@ public:
 	virtual FText GetNodeTitle(ENodeTitleType::Type TitleType) const override;
 	// End of UEdGraphNode interface
 
-	// UAnimGraphNode_SkeletalControlBase interface
-	ANIMGRAPH_API virtual FVector GetWidgetLocation(const USkeletalMeshComponent* SkelComp, struct FAnimNode_SkeletalControlBase* AnimNode) override;
-	ANIMGRAPH_API virtual int32 GetWidgetMode(const USkeletalMeshComponent* SkelComp) override;
-	ANIMGRAPH_API virtual FName FindSelectedBone() override;
-	ANIMGRAPH_API virtual void DoTranslation(const USkeletalMeshComponent* SkelComp, FVector& Drag, FAnimNode_Base* InOutAnimNode) override;
-
-	ANIMGRAPH_API virtual void CopyNodeDataTo(FAnimNode_Base* AnimNode) override;
-	ANIMGRAPH_API virtual void CopyNodeDataFrom(const FAnimNode_Base* NewAnimNode) override;
-	ANIMGRAPH_API virtual void Draw(FPrimitiveDrawInterface* PDI, USkeletalMeshComponent * PreviewSkelMeshComp) const override;
-	// End of UAnimGraphNode_SkeletalControlBase interface
+	// UAnimGraphNode_Base interface
+	virtual void CopyNodeDataToPreviewNode(FAnimNode_Base* AnimNode) override;
+	virtual FEditorModeID GetEditorMode() const override;
+	virtual void Draw(FPrimitiveDrawInterface* PDI, USkeletalMeshComponent * PreviewSkelMeshComp) const override;
+	// End of UAnimGraphNode_Base interface
 
 protected:
 	// UAnimGraphNode_SkeletalControlBase interface
diff --git a/Engine/Source/Editor/AnimGraph/Classes/AnimGraphNode_LegIK.h b/Engine/Source/Editor/AnimGraph/Classes/AnimGraphNode_LegIK.h
--- a/Engine/Source/Editor/AnimGraph/Classes/AnimGraphNode_LegIK.h
+++ b/Engine/Source/Editor/AnimGraph/Classes/AnimGraphNode_LegIK.h
@@ -3,7 +3,7 @@
 #pragma once
 
 #include "AnimGraphNode_SkeletalControlBase.h"
-#include "Private/BoneControllers/AnimNode_LegIK.h"
+#include "BoneControllers/AnimNode_LegIK.h"
 
 #include "AnimGraphNode_LegIK.generated.h"
 
diff --git a/Engine/Source/Editor/AnimGraph/Classes/AnimGraphNode_ModifyBone.h b/Engine/Source/Editor/AnimGraph/Classes/AnimGraphNode_ModifyBone.h
--- a/Engine/Source/Editor/AnimGraph/Classes/AnimGraphNode_ModifyBone.h
+++ b/Engine/Source/Editor/AnimGraph/Classes/AnimGraphNode_ModifyBone.h
@@ -24,29 +24,16 @@ public:
 protected:
 	// UAnimGraphNode_Base interface
 	virtual void ValidateAnimNodeDuringCompilation(USkeleton* ForSkeleton, FCompilerResultsLog& MessageLog) override;
+	virtual FEditorModeID GetEditorMode() const override;
+	virtual void CopyNodeDataToPreviewNode(FAnimNode_Base* InPreviewNode) override;
+	virtual void CopyPinDefaultsToNodeData(UEdGraphPin* InPin) override;
 	// End of UAnimGraphNode_Base interface
 
 	// UAnimGraphNode_SkeletalControlBase interface
 	virtual FText GetControllerDescription() const override;
-	virtual int32 GetWidgetCoordinateSystem(const USkeletalMeshComponent* SkelComp) override;
-	virtual FVector GetWidgetLocation(const USkeletalMeshComponent* SkelComp, struct FAnimNode_SkeletalControlBase* AnimNode) override;
-	virtual int32 GetWidgetMode(const USkeletalMeshComponent* SkelComp) override;
-	virtual int32 ChangeToNextWidgetMode(const USkeletalMeshComponent* SkelComp, int32 InCurWidgetMode) override;
-	virtual bool SetWidgetMode(const USkeletalMeshComponent* SkelComp, int32 InWidgetMode) override;
-	virtual FName FindSelectedBone() override;
-	virtual void DoTranslation(const USkeletalMeshComponent* SkelComp, FVector& Drag, FAnimNode_Base* InOutAnimNode) override;
-	virtual void DoRotation(const USkeletalMeshComponent* SkelComp, FRotator& Rotation, FAnimNode_Base* InOutAnimNode) override;
-	virtual void DoScale(const USkeletalMeshComponent* SkelComp, FVector& Drag, FAnimNode_Base* InOutAnimNode) override;
-	virtual void CopyNodeDataTo(FAnimNode_Base* OutAnimNode) override;
-	virtual void CopyNodeDataFrom(const FAnimNode_Base* InNewAnimNode) override;
 	virtual const FAnimNode_SkeletalControlBase* GetNode() const override { return &Node; }
 	// End of UAnimGraphNode_SkeletalControlBase interface
 
-	// methods to find a valid widget mode for gizmo because doesn't need to show gizmo when the mode is "Ignore"
-	int32 FindValidWidgetMode(int32 InWidgetMode);
-	EBoneModificationMode GetBoneModificationMode(int32 InWidgetMode);
-	int32 GetNextWidgetMode(int32 InWidgetMode);
-
 private:
 	/** Constructing FText strings can be costly, so we cache the node's title */
 	FNodeTitleTextTable CachedNodeTitles;
diff --git a/Engine/Source/Editor/AnimGraph/Classes/AnimGraphNode_PoseBlendNode.h b/Engine/Source/Editor/AnimGraph/Classes/AnimGraphNode_PoseBlendNode.h
--- a/Engine/Source/Editor/AnimGraph/Classes/AnimGraphNode_PoseBlendNode.h
+++ b/Engine/Source/Editor/AnimGraph/Classes/AnimGraphNode_PoseBlendNode.h
@@ -1,14 +1,13 @@
 // Copyright 1998-2016 Epic Games, Inc. All Rights Reserved.
 
 #pragma once
-#include "AnimGraphNode_Base.h"
-#include "AnimGraphNode_AssetPlayerBase.h"
+#include "AnimGraphNode_PoseHandler.h"
 #include "AnimNodes/AnimNode_PoseBlendNode.h"
 #include "EdGraph/EdGraphNodeUtils.h" // for FNodeTextCache
 #include "AnimGraphNode_PoseBlendNode.generated.h"
 
 UCLASS(MinimalAPI)
-class UAnimGraphNode_PoseBlendNode : public UAnimGraphNode_AssetPlayerBase
+class UAnimGraphNode_PoseBlendNode : public UAnimGraphNode_PoseHandler
 {
 	GENERATED_UCLASS_BODY()
 
@@ -22,11 +21,8 @@ class UAnimGraphNode_PoseBlendNode : public UAnimGraphNode_AssetPlayerBase
 	// End of UEdGraphNode
 
 	// UAnimGraphNode_Base interface
-	virtual void ValidateAnimNodeDuringCompilation(class USkeleton* ForSkeleton, class FCompilerResultsLog& MessageLog) override;
-	virtual void PreloadRequiredAssets() override;
 	// Interface to support transition getter
 	virtual bool DoesSupportTimeForTransitionGetter() const override;
-	virtual UAnimationAsset* GetAnimationAsset() const override;
 	virtual void GetAllAnimationSequencesReferred(TArray<UAnimationAsset*>& AnimationAssets) const override;
 	virtual void ReplaceReferredAnimations(const TMap<UAnimationAsset*, UAnimationAsset*>& AnimAssetReplacementMap) override;
 	virtual void GetMenuActions(FBlueprintActionDatabaseRegistrar& ActionRegistrar) const override;
@@ -37,17 +33,16 @@ class UAnimGraphNode_PoseBlendNode : public UAnimGraphNode_AssetPlayerBase
 	virtual void GetContextMenuActions(const FGraphNodeContextMenuBuilder& Context) const override;
 	// End of UK2Node interface
 
-	virtual void SetAnimationAsset(UAnimationAsset* Asset) override;
-
 private:
 	FText GetNodeTitleForPoseAsset(ENodeTitleType::Type TitleType, UPoseAsset* InPoseAsset) const;
 
 private:
 	static FText GetTitleGivenAssetInfo(const FText& AssetName);
 
-	/** Constructing FText strings can be costly, so we cache the node's title */
-	FNodeTextCache CachedNodeTitle;
-
 	/** Used for filtering in the Blueprint context menu when the pose asset this node uses is unloaded */
 	FString UnloadedSkeletonName;
+
+protected:
+	virtual FAnimNode_PoseHandler* GetPoseHandlerNode() override { return &Node; }
+	virtual const FAnimNode_PoseHandler* GetPoseHandlerNode() const override { return &Node; }
 };
diff --git a/Engine/Source/Editor/AnimGraph/Classes/AnimGraphNode_PoseDriver.h b/Engine/Source/Editor/AnimGraph/Classes/AnimGraphNode_PoseDriver.h
--- a/Engine/Source/Editor/AnimGraph/Classes/AnimGraphNode_PoseDriver.h
+++ b/Engine/Source/Editor/AnimGraph/Classes/AnimGraphNode_PoseDriver.h
@@ -5,13 +5,12 @@
 class IDetailCategoryBuilder;
 class IPropertyHandle;
 
-#include "AnimGraphNode_SkeletalControlBase.h"
-#include "BoneControllers/AnimNode_PoseDriver.h"
-#include "EdGraph/EdGraphNodeUtils.h" // for FNodeTitleTextTable
+#include "AnimGraphNode_PoseHandler.h"
+#include "AnimNodes/AnimNode_PoseDriver.h"
 #include "AnimGraphNode_PoseDriver.generated.h"
 
 UCLASS()
-class ANIMGRAPH_API UAnimGraphNode_PoseDriver : public UAnimGraphNode_SkeletalControlBase
+class ANIMGRAPH_API UAnimGraphNode_PoseDriver : public UAnimGraphNode_PoseHandler
 {
 	GENERATED_UCLASS_BODY()
 
@@ -23,22 +22,17 @@ public:
 	// UEdGraphNode interface
 	virtual FText GetNodeTitle(ENodeTitleType::Type TitleType) const override;
 	virtual FText GetTooltipText() const override;
+	virtual FText GetMenuCategory() const override;
 	// End of UEdGraphNode interface
 
 	// UAnimGraphNode_Base interface
 	virtual void ValidateAnimNodeDuringCompilation(USkeleton* ForSkeleton, FCompilerResultsLog& MessageLog) override;
-	virtual void CustomizeDetails(IDetailLayoutBuilder& DetailBuilder) override;
+	virtual FEditorModeID GetEditorMode() const override;
 	// End of UAnimGraphNode_Base interface
 
-	// UAnimGraphNode_SkeletalControlBase interface
-	virtual void Draw(FPrimitiveDrawInterface* PDI, USkeletalMeshComponent* SkelMeshComp) const override;
-	virtual void DrawCanvas(FViewport& InViewport, FSceneView& View, FCanvas& Canvas, USkeletalMeshComponent * PreviewSkelMeshComp) const override;
-	// End of UAnimGraphNode_SkeletalControlBase interface
-
 protected:
-
-	// UAnimGraphNode_SkeletalControlBase protected interface
-	virtual FText GetControllerDescription() const override;
-	virtual const FAnimNode_SkeletalControlBase* GetNode() const override { return &Node; }
-	// End of UAnimGraphNode_SkeletalControlBase protected interface
+	// UAnimGraphNode_PoseHandler interface
+	virtual FAnimNode_PoseHandler* GetPoseHandlerNode() override { return &Node; }
+	virtual const FAnimNode_PoseHandler* GetPoseHandlerNode() const override { return &Node; }
+	// End of UAnimGraphNode_PoseHandler interface
 };
diff --git a/Engine/Source/Editor/AnimGraph/Classes/AnimGraphNode_SkeletalControlBase.h b/Engine/Source/Editor/AnimGraph/Classes/AnimGraphNode_SkeletalControlBase.h
--- a/Engine/Source/Editor/AnimGraph/Classes/AnimGraphNode_SkeletalControlBase.h
+++ b/Engine/Source/Editor/AnimGraph/Classes/AnimGraphNode_SkeletalControlBase.h
@@ -2,6 +2,7 @@
 
 #pragma once
 
+#include "Editor.h"
 #include "AnimGraphNode_Base.h"
 #include "AnimGraphNode_SkeletalControlBase.generated.h"
 
@@ -30,74 +31,66 @@ public:
 	virtual void ValidateAnimNodePostCompile(FCompilerResultsLog& MessageLog, UAnimBlueprintGeneratedClass* CompiledClass, int32 CompiledNodeIndex) override;
 	// End of UAnimGraphNode_Base interface
 
-	// Draw function for supporting visualization
-	virtual void Draw(FPrimitiveDrawInterface* PDI, USkeletalMeshComponent * PreviewSkelMeshComp) const {}
-	// Canvas draw function to draw to viewport
-	virtual void DrawCanvas(FViewport& InViewport, FSceneView& View, FCanvas& Canvas, USkeletalMeshComponent * PreviewSkelMeshComp) const {}
-	// Function to collect strings from nodes to display in the viewport.
-	// Use this rather than DrawCanvas when adding general text to the viewport.
-	virtual void GetOnScreenDebugInfo(TArray<FText>& DebugInfo, USkeletalMeshComponent* PreviewSkelMeshComp) const {}
-
 	/**
 	 * methods related to widget control
 	 */
+	DEPRECATED(4.14, "This function is deprecated. Use a IAnimNodeEditMode-derived class to implement a new editor mode and override GetEditorMode to return its ID")
 	virtual FVector GetWidgetLocation(const USkeletalMeshComponent* SkelComp, struct FAnimNode_SkeletalControlBase* AnimNode)
 	{
 		return FVector::ZeroVector;
 	}
-	// to keep data consistency between anim nodes 
-	virtual void CopyNodeDataTo(FAnimNode_Base* OutAnimNode){}
+	
+	DEPRECATED(4.14, "This function is deprecated and is no longer used, please use CopyNodeDataToPreviewNode")
+	virtual void CopyNodeDataTo(FAnimNode_Base* OutAnimNode){ CopyNodeDataToPreviewNode(OutAnimNode); }
+
+	DEPRECATED(4.14, "This function is deprecated and is no longer used")
 	virtual void CopyNodeDataFrom(const FAnimNode_Base* NewAnimNode){}
 
 	/** Are we currently showing this pin */
 	bool IsPinShown(const FString& PinName) const;
 
 	// Returns the coordinate system that should be used for this bone
+	DEPRECATED(4.14, "This function is deprecated. Use a IAnimNodeEditMode-derived class to implement a new editor mode and override GetEditorMode to return its ID")
 	virtual int32 GetWidgetCoordinateSystem(const USkeletalMeshComponent* SkelComp);
 
 	// return current widget mode this anim graph node supports
+	DEPRECATED(4.14, "This function is deprecated. Use a IAnimNodeEditMode-derived class to implement a new editor mode and override GetEditorMode to return its ID")
 	virtual int32 GetWidgetMode(const USkeletalMeshComponent* SkelComp);
 	// called when the user changed widget mode by pressing "Space" key
+	DEPRECATED(4.14, "This function is deprecated. Use a IAnimNodeEditMode-derived class to implement a new editor mode and override GetEditorMode to return its ID")
 	virtual int32 ChangeToNextWidgetMode(const USkeletalMeshComponent* SkelComp, int32 CurWidgetMode);
 	// called when the user set widget mode directly, returns true if InWidgetMode is available
+	DEPRECATED(4.14, "This function is deprecated. Use a IAnimNodeEditMode-derived class to implement a new editor mode and override GetEditorMode to return its ID")
 	virtual bool SetWidgetMode(const USkeletalMeshComponent* SkelComp, int32 InWidgetMode) { return false; }
 
 	// 
+	DEPRECATED(4.14, "This function is deprecated. Use a IAnimNodeEditMode-derived class to implement a new editor mode and override GetEditorMode to return its ID")
 	virtual FName FindSelectedBone();
 
 	// if anim graph node needs other actors to select other bones, move actor's positions when this is called
+	DEPRECATED(4.14, "This function is deprecated. Use a IAnimNodeEditMode-derived class to implement a new editor mode and override GetEditorMode to return its ID")
 	virtual void MoveSelectActorLocation(const USkeletalMeshComponent* SkelComp, FAnimNode_SkeletalControlBase* InAnimNode) {}
 
+	DEPRECATED(4.14, "This function is deprecated. Use a IAnimNodeEditMode-derived class to implement a new editor mode and override GetEditorMode to return its ID")
 	virtual bool IsActorClicked(HActor* ActorHitProxy) { return false; }
+	DEPRECATED(4.14, "This function is deprecated. Use a IAnimNodeEditMode-derived class to implement a new editor mode and override GetEditorMode to return its ID")
 	virtual void ProcessActorClick(HActor* ActorHitProxy) {}
 	// if it has select-actors, should hide all actors when de-select is called  
+	DEPRECATED(4.14, "This function is deprecated. Use a IAnimNodeEditMode-derived class to implement a new editor mode and override GetEditorMode to return its ID")
 	virtual void DeselectActor(USkeletalMeshComponent* SkelComp){}
 
 	// called when the widget is dragged in translation mode
+	DEPRECATED(4.14, "This function is deprecated. Use a IAnimNodeEditMode-derived class to implement a new editor mode and override GetEditorMode to return its ID")
 	virtual void DoTranslation(const USkeletalMeshComponent* SkelComp, FVector& Drag, FAnimNode_Base* InOutAnimNode) {}
 	// called when the widget is dragged in rotation mode
+	DEPRECATED(4.14, "This function is deprecated. Use a IAnimNodeEditMode-derived class to implement a new editor mode and override GetEditorMode to return its ID")
 	virtual void DoRotation(const USkeletalMeshComponent* SkelComp, FRotator& Rotation, FAnimNode_Base* InOutAnimNode) {}
 	// called when the widget is dragged in scale mode
+	DEPRECATED(4.14, "This function is deprecated. Use a IAnimNodeEditMode-derived class to implement a new editor mode and override GetEditorMode to return its ID")
 	virtual void DoScale(const USkeletalMeshComponent* SkelComp, FVector& Scale, FAnimNode_Base* InOutAnimNode) {}
 
-	/** Try to find the preview control node instance for this anim graph node */
-	FAnimNode_SkeletalControlBase* FindDebugAnimNode(USkeletalMeshComponent * PreviewSkelMeshComp) const;
 
-protected:
-	// Returns the short descriptive name of the controller
-	virtual FText GetControllerDescription() const;
-
-	/**
-	* helper functions for bone control preview
-	*/
-	// local conversion function for drawing
-	void ConvertToComponentSpaceTransform(const USkeletalMeshComponent* SkelComp, const FTransform & InTransform, FTransform & OutCSTransform, int32 BoneIndex, EBoneControlSpace Space) const;
-	// convert drag vector in component space to bone space 
-	FVector ConvertCSVectorToBoneSpace(const USkeletalMeshComponent* SkelComp, FVector& InCSVector, FCSPose<FCompactHeapPose>& MeshBases, const FName& BoneName, const EBoneControlSpace Space);
-	// convert rotator in component space to bone space 
-	FQuat ConvertCSRotationToBoneSpace(const USkeletalMeshComponent* SkelComp, FRotator& InCSRotator, FCSPose<FCompactHeapPose>& MeshBases, const FName& BoneName, const EBoneControlSpace Space);
-	// convert widget location according to bone control space
-	FVector ConvertWidgetLocation(const USkeletalMeshComponent* InSkelComp, FCSPose<FCompactHeapPose>& InMeshBases, const FName& BoneName, const FVector& InLocation, const EBoneControlSpace Space);
+public:
 	// set literal value for FVector
 	void SetDefaultValue(const FString& InDefaultValueName, const FVector& InValue);
 	// get literal value for vector
@@ -140,5 +133,25 @@ protected:
 		CompileNodeValue = InValue;
 	}
 
+protected:
+	// Returns the short descriptive name of the controller
+	virtual FText GetControllerDescription() const;
+
+	/**
+	* helper functions for bone control preview
+	*/
+	// local conversion function for drawing
+	DEPRECATED(4.14, "This function is moved to FAnimNodeEditMode. Use a IAnimNodeEditMode-derived class to implement a new editor mode and override GetEditorMode to return its ID")
+	void ConvertToComponentSpaceTransform(const USkeletalMeshComponent* SkelComp, const FTransform & InTransform, FTransform & OutCSTransform, int32 BoneIndex, EBoneControlSpace Space) const;
+	// convert drag vector in component space to bone space 
+	DEPRECATED(4.14, "This function is moved to FAnimNodeEditMode. Use a IAnimNodeEditMode-derived class to implement a new editor mode and override GetEditorMode to return its ID")
+	FVector ConvertCSVectorToBoneSpace(const USkeletalMeshComponent* SkelComp, FVector& InCSVector, FCSPose<FCompactHeapPose>& MeshBases, const FName& BoneName, const EBoneControlSpace Space);
+	// convert rotator in component space to bone space 
+	DEPRECATED(4.14, "This function is moved to FAnimNodeEditMode. Use a IAnimNodeEditMode-derived class to implement a new editor mode and override GetEditorMode to return its ID")
+	FQuat ConvertCSRotationToBoneSpace(const USkeletalMeshComponent* SkelComp, FRotator& InCSRotator, FCSPose<FCompactHeapPose>& MeshBases, const FName& BoneName, const EBoneControlSpace Space);
+	// convert widget location according to bone control space
+	DEPRECATED(4.14, "This function is moved to FAnimNodeEditMode. Use a IAnimNodeEditMode-derived class to implement a new editor mode and override GetEditorMode to return its ID")
+	FVector ConvertWidgetLocation(const USkeletalMeshComponent* InSkelComp, FCSPose<FCompactHeapPose>& InMeshBases, const FName& BoneName, const FVector& InLocation, const EBoneControlSpace Space);
+
 	virtual const FAnimNode_SkeletalControlBase* GetNode() const PURE_VIRTUAL(UAnimGraphNode_SkeletalControlBase::GetNode, return nullptr;);
 };
diff --git a/Engine/Source/Editor/AnimGraph/Classes/AnimGraphNode_TwoBoneIK.h b/Engine/Source/Editor/AnimGraph/Classes/AnimGraphNode_TwoBoneIK.h
--- a/Engine/Source/Editor/AnimGraph/Classes/AnimGraphNode_TwoBoneIK.h
+++ b/Engine/Source/Editor/AnimGraph/Classes/AnimGraphNode_TwoBoneIK.h
@@ -26,43 +26,17 @@ public:
 	virtual FText GetTooltipText() const override;
 	// End of UEdGraphNode interface
 
-	void UpdateEffectorLocationSpace(class IDetailLayoutBuilder& DetailBuilder);
 	// UAnimGraphNode_Base interface
 	virtual void CustomizeDetails(class IDetailLayoutBuilder& DetailBuilder) override;
+	virtual FEditorModeID GetEditorMode() const;
+	virtual void CopyNodeDataToPreviewNode(FAnimNode_Base* InPreviewNode) override;
+	virtual void CopyPinDefaultsToNodeData(UEdGraphPin* InPin) override;
 	// End of UAnimGraphNode_Base interface
 
 	// UAnimGraphNode_SkeletalControlBase interface
-	ANIMGRAPH_API virtual void Draw( FPrimitiveDrawInterface* PDI, USkeletalMeshComponent* SkelMeshComp) const override;
-
-	ANIMGRAPH_API virtual FVector GetWidgetLocation(const USkeletalMeshComponent* SkelComp, struct FAnimNode_SkeletalControlBase* AnimNode) override;
-
-	ANIMGRAPH_API virtual void	CopyNodeDataFrom(const FAnimNode_Base* NewAnimNode) override;
-
-	ANIMGRAPH_API virtual int32 GetWidgetMode(const USkeletalMeshComponent* SkelComp) override;
-
-	ANIMGRAPH_API virtual void MoveSelectActorLocation(const USkeletalMeshComponent* SkelComp, FAnimNode_SkeletalControlBase* AnimNode) override;
-	ANIMGRAPH_API virtual FName FindSelectedBone() override;
-	ANIMGRAPH_API virtual bool IsActorClicked(HActor* ActorHitProxy) override;
-	ANIMGRAPH_API virtual void ProcessActorClick(HActor* ActorHitProxy) override;
-	ANIMGRAPH_API virtual void DoTranslation(const USkeletalMeshComponent* SkelComp, FVector& Drag, FAnimNode_Base* InOutAnimNode) override;
-	ANIMGRAPH_API virtual void	CopyNodeDataTo(FAnimNode_Base* AnimNode) override;
-	ANIMGRAPH_API virtual void	DeselectActor(USkeletalMeshComponent* SkelComp) override;
-
 	virtual const FAnimNode_SkeletalControlBase* GetNode() const override { return &Node; }
 	// End of UAnimGraphNode_SkeletalControlBase interface
 
-	enum BoneSelectModeType
-	{
-		BSM_EndEffector,
-		BSM_JointTarget,
-		BSM_Max
-	};
-
-	// to switch select mode between end effector and joint target
-	BoneSelectModeType BoneSelectMode;
-
-	TWeakObjectPtr<class ABoneSelectActor> BoneSelectActor;
-
 	IDetailLayoutBuilder* DetailLayout;
 
 protected:
@@ -70,9 +44,6 @@ protected:
 	virtual FText GetControllerDescription() const override;
 	// End of UAnimGraphNode_SkeletalControlBase interface
 
-	// local conversion function for drawing
-	void DrawTargetLocation(FPrimitiveDrawInterface* PDI, USkeletalMeshComponent* SkelComp, USkeleton * Skeleton, EBoneControlSpace SpaceBase, FName SpaceBoneName, const FVector & TargetLocation, const FColor & TargetColor, const FColor & BoneColor) const;
-
 private:
 	/** Constructing FText strings can be costly, so we cache the node's title */
 	FNodeTitleTextTable CachedNodeTitles;
diff --git a/Engine/Source/Editor/AnimGraph/Classes/AnimationGraph.h b/Engine/Source/Editor/AnimGraph/Classes/AnimationGraph.h
--- a/Engine/Source/Editor/AnimGraph/Classes/AnimationGraph.h
+++ b/Engine/Source/Editor/AnimGraph/Classes/AnimationGraph.h
@@ -4,9 +4,15 @@
 
 #include "AnimationGraph.generated.h"
 
+/** Delegate fired when a pin's default value is changed */
+DECLARE_MULTICAST_DELEGATE_OneParam(FOnPinDefaultValueChanged, UEdGraphPin* /*InPinThatChanged*/)
+
 UCLASS(MinimalAPI)
 class UAnimationGraph : public UEdGraph
 {
 	GENERATED_UCLASS_BODY()
+
+	/** Delegate fired when a pin's default value is changed */
+	FOnPinDefaultValueChanged OnPinDefaultValueChanged;
 };
 
diff --git a/Engine/Source/Editor/AnimGraph/Classes/AnimationGraphSchema.h b/Engine/Source/Editor/AnimGraph/Classes/AnimationGraphSchema.h
--- a/Engine/Source/Editor/AnimGraph/Classes/AnimationGraphSchema.h
+++ b/Engine/Source/Editor/AnimGraph/Classes/AnimationGraphSchema.h
@@ -52,6 +52,7 @@ class UAnimationGraphSchema : public UEdGraphSchema_K2
 	virtual bool CanDuplicateGraph(UEdGraph* InSourceGraph) const override {	return false; }
 	virtual bool DoesSupportEventDispatcher() const	override { return false; }
 	virtual bool ShouldAlwaysPurgeOnModification() const override { return true; }
+	virtual bool TryCreateConnection(UEdGraphPin* A, UEdGraphPin* B) const override;
 	//~ End UEdGraphSchema Interface.
 
 	//~ Begin UEdGraphSchema_K2 Interface
diff --git a/Engine/Source/Editor/AnimGraph/Private/AnimGraphModule.cpp b/Engine/Source/Editor/AnimGraph/Private/AnimGraphModule.cpp
--- a/Engine/Source/Editor/AnimGraph/Private/AnimGraphModule.cpp
+++ b/Engine/Source/Editor/AnimGraph/Private/AnimGraphModule.cpp
@@ -2,7 +2,40 @@
 
 #include "AnimGraphPrivatePCH.h"
 #include "AnimGraphModule.h"
+#include "AnimNodeEditModes.h"
+#include "AnimNodeEditMode.h"
+#include "EditModes/TwoBoneIKEditMode.h"
+#include "EditModes/ObserveBoneEditMode.h"
+#include "EditModes/ModifyBoneEditMode.h"
+#include "EditModes/FabrikEditMode.h"
+#include "EditModes/PoseDriverEditMode.h"
 
+IMPLEMENT_MODULE(FAnimGraphModule, AnimGraph);
 
-IMPLEMENT_MODULE( FAnimGraphModule, AnimGraph );
+#define LOCTEXT_NAMESPACE "AnimGraphModule"
 
+void FAnimGraphModule::StartupModule()
+{
+	FAnimGraphCommands::Register();
+
+	// Register the editor modes
+	FEditorModeRegistry::Get().RegisterMode<FAnimNodeEditMode>(AnimNodeEditModes::AnimNode, LOCTEXT("AnimNodeEditMode", "Anim Node"), FSlateIcon(), false);
+	FEditorModeRegistry::Get().RegisterMode<FTwoBoneIKEditMode>(AnimNodeEditModes::TwoBoneIK, LOCTEXT("TwoBoneIKEditMode", "2-Bone IK"), FSlateIcon(), false);
+	FEditorModeRegistry::Get().RegisterMode<FObserveBoneEditMode>(AnimNodeEditModes::ObserveBone, LOCTEXT("ObserveBoneEditMode", "Observe Bone"), FSlateIcon(), false);
+	FEditorModeRegistry::Get().RegisterMode<FModifyBoneEditMode>(AnimNodeEditModes::ModifyBone, LOCTEXT("ModifyBoneEditMode", "Modify Bone"), FSlateIcon(), false);
+	FEditorModeRegistry::Get().RegisterMode<FFabrikEditMode>(AnimNodeEditModes::Fabrik, LOCTEXT("FabrikEditMode", "Fabrik"), FSlateIcon(), false);
+	FEditorModeRegistry::Get().RegisterMode<FPoseDriverEditMode>(AnimNodeEditModes::PoseDriver, LOCTEXT("PoseDriverEditMode", "PoseDriver"), FSlateIcon(), false);
+}
+
+void FAnimGraphModule::ShutdownModule()
+{
+	// Unregister the editor modes
+	FEditorModeRegistry::Get().UnregisterMode(AnimNodeEditModes::PoseDriver);
+	FEditorModeRegistry::Get().UnregisterMode(AnimNodeEditModes::Fabrik);
+	FEditorModeRegistry::Get().UnregisterMode(AnimNodeEditModes::ModifyBone);
+	FEditorModeRegistry::Get().UnregisterMode(AnimNodeEditModes::ObserveBone);
+	FEditorModeRegistry::Get().UnregisterMode(AnimNodeEditModes::TwoBoneIK);
+	FEditorModeRegistry::Get().UnregisterMode(AnimNodeEditModes::AnimNode);
+}
+
+#undef LOCTEXT_NAMESPACE
\ No newline at end of file
diff --git a/Engine/Source/Editor/AnimGraph/Private/AnimGraphNode_AssetPlayerBase.cpp b/Engine/Source/Editor/AnimGraph/Private/AnimGraphNode_AssetPlayerBase.cpp
new file mode 100644
--- /dev/null
+++ b/Engine/Source/Editor/AnimGraph/Private/AnimGraphNode_AssetPlayerBase.cpp
@@ -0,0 +1,124 @@
+// Copyright 1998-2016 Epic Games, Inc. All Rights Reserved.
+
+#include "AnimGraphPrivatePCH.h"
+#include "Animation/AnimComposite.h"
+#include "Animation/AimOffsetBlendSpace.h"
+#include "Animation/AimOffsetBlendSpace1D.h"
+#include "AnimGraphNode_AssetPlayerBase.h"
+#include "AnimGraphNode_SequencePlayer.h"
+#include "AnimGraphNode_SequenceEvaluator.h"
+#include "AnimGraphNode_RotationOffsetBlendSpace.h"
+#include "AnimGraphNode_BlendSpacePlayer.h"
+#include "AnimGraphNode_BlendSpaceEvaluator.h"
+#include "AnimGraphNode_PoseBlendNode.h"
+#include "AnimGraphNode_PoseByName.h"
+#include "AnimGraphNode_PoseDriver.h"
+
+void UAnimGraphNode_AssetPlayerBase::PinConnectionListChanged(UEdGraphPin* Pin)
+{
+	Super::PinConnectionListChanged(Pin);
+
+	if (Pin->PinType.PinCategory == UEdGraphSchema_K2::PC_Object)
+	{
+		// if pin has no connections, we may need to restore the default
+		if (Pin->LinkedTo.Num() == 0)
+		{
+			UObject* Asset = GetAssetReferenceForPinRestoration();
+			if (Asset != nullptr)
+			{
+				Pin->DefaultObject = Asset;
+			}
+		}
+
+		// recache visualization now an asset pin's connection is changed
+		if (const UEdGraphSchema* Schema = GetSchema())
+		{
+			Schema->ForceVisualizationCacheClear();
+		}
+	}
+}
+
+void UAnimGraphNode_AssetPlayerBase::PinDefaultValueChanged(UEdGraphPin* Pin)
+{
+	Super::PinDefaultValueChanged(Pin);
+
+	if (Pin->PinType.PinCategory == UEdGraphSchema_K2::PC_Object)
+	{
+		// recache visualization now an asset pin's default value has changed
+		if (const UEdGraphSchema* Schema = GetSchema())
+		{
+			Schema->ForceVisualizationCacheClear();
+		}
+	}
+}
+
+void UAnimGraphNode_AssetPlayerBase::SetAssetReferenceForPinRestoration(UObject* InAsset)
+{
+	AssetReferenceForPinRestoration = InAsset;
+}
+
+UObject* UAnimGraphNode_AssetPlayerBase::GetAssetReferenceForPinRestoration()
+{
+	return AssetReferenceForPinRestoration.TryLoad();
+}
+
+bool IsAimOffsetBlendSpace(const UClass* BlendSpaceClass)
+{
+	return  BlendSpaceClass->IsChildOf(UAimOffsetBlendSpace::StaticClass()) ||
+		BlendSpaceClass->IsChildOf(UAimOffsetBlendSpace1D::StaticClass());
+}
+
+UClass* GetNodeClassForAsset(const UClass* AssetClass)
+{
+	if (AssetClass->IsChildOf(UAnimSequence::StaticClass()))
+	{
+		return UAnimGraphNode_SequencePlayer::StaticClass();
+	}
+	else if (AssetClass->IsChildOf(UBlendSpaceBase::StaticClass()))
+	{
+		if (IsAimOffsetBlendSpace(AssetClass))
+		{
+			return UAnimGraphNode_RotationOffsetBlendSpace::StaticClass();
+		}
+		else
+		{
+			return UAnimGraphNode_BlendSpacePlayer::StaticClass();
+		}
+	}
+	else if (AssetClass->IsChildOf(UAnimComposite::StaticClass()))
+	{
+		return UAnimGraphNode_SequencePlayer::StaticClass();
+	}
+	else if (AssetClass->IsChildOf(UPoseAsset::StaticClass()))
+	{
+		return UAnimGraphNode_PoseBlendNode::StaticClass();
+	}
+	return nullptr;
+}
+
+bool SupportNodeClassForAsset(const UClass* AssetClass, const UClass* NodeClass)
+{
+	// we don't want montage to show up, so not checking AnimSequenceBase
+	if (AssetClass->IsChildOf(UAnimSequence::StaticClass()) || AssetClass->IsChildOf(UAnimComposite::StaticClass()))
+	{
+		return (UAnimGraphNode_SequencePlayer::StaticClass() == NodeClass || UAnimGraphNode_SequenceEvaluator::StaticClass() == NodeClass);
+	}
+	else if (AssetClass->IsChildOf(UBlendSpaceBase::StaticClass()))
+	{
+		if (IsAimOffsetBlendSpace(AssetClass))
+		{
+			return (UAnimGraphNode_RotationOffsetBlendSpace::StaticClass() == NodeClass);
+		}
+		else
+		{
+			return (UAnimGraphNode_BlendSpacePlayer::StaticClass() == NodeClass || UAnimGraphNode_BlendSpaceEvaluator::StaticClass() == NodeClass);
+		}
+	}
+	else if (AssetClass->IsChildOf(UPoseAsset::StaticClass()))
+	{
+		return (UAnimGraphNode_PoseBlendNode::StaticClass() == NodeClass || UAnimGraphNode_PoseByName::StaticClass() == NodeClass || UAnimGraphNode_PoseDriver::StaticClass());
+	}
+	
+	return false;
+}
+
diff --git a/Engine/Source/Editor/AnimGraph/Private/AnimGraphNode_Base.cpp b/Engine/Source/Editor/AnimGraph/Private/AnimGraphNode_Base.cpp
--- a/Engine/Source/Editor/AnimGraph/Private/AnimGraphNode_Base.cpp
+++ b/Engine/Source/Editor/AnimGraph/Private/AnimGraphNode_Base.cpp
@@ -8,6 +8,9 @@
 #include "BlueprintActionDatabaseRegistrar.h"
 #include "Animation/AnimInstance.h"
 #include "AnimBlueprintNodeOptionalPinManager.h"
+#include "IAnimNodeEditMode.h"
+#include "AnimNodeEditModes.h"
+#include "AnimationGraph.h"
 
 /////////////////////////////////////////////////////
 // UAnimGraphNode_Base
@@ -268,3 +271,68 @@ void UAnimGraphNode_Base::HandleAnimReferenceCollection(UAnimationAsset* AnimAss
 		AnimAsset->HandleAnimReferenceCollection(AnimationAssets, true);
 	}
 }
+
+void UAnimGraphNode_Base::OnNodeSelected(bool bInIsSelected, FEditorModeTools& InModeTools, FAnimNode_Base* InRuntimeNode)
+{
+	const FEditorModeID ModeID = GetEditorMode();
+	if (ModeID != NAME_None)
+	{
+		if (bInIsSelected)
+		{
+			InModeTools.ActivateMode(ModeID);
+			if (FEdMode* EdMode = InModeTools.GetActiveMode(ModeID))
+			{
+				static_cast<IAnimNodeEditMode*>(EdMode)->EnterMode(this, InRuntimeNode);
+			}
+		}
+		else
+		{
+			if (FEdMode* EdMode = InModeTools.GetActiveMode(ModeID))
+			{
+				static_cast<IAnimNodeEditMode*>(EdMode)->ExitMode();
+			}
+			InModeTools.DeactivateMode(ModeID);
+		}
+	}
+}
+
+FEditorModeID UAnimGraphNode_Base::GetEditorMode() const
+{
+	return AnimNodeEditModes::AnimNode;
+}
+
+FAnimNode_Base* UAnimGraphNode_Base::FindDebugAnimNode(USkeletalMeshComponent * PreviewSkelMeshComp) const
+{
+	FAnimNode_Base* DebugNode = nullptr;
+
+	if (PreviewSkelMeshComp != nullptr && PreviewSkelMeshComp->GetAnimInstance() != nullptr)
+	{
+		// find an anim node index from debug data
+		UAnimBlueprintGeneratedClass* AnimBlueprintClass = Cast<UAnimBlueprintGeneratedClass>(PreviewSkelMeshComp->GetAnimInstance()->GetClass());
+		if (AnimBlueprintClass)
+		{
+			FAnimBlueprintDebugData& DebugData = AnimBlueprintClass->GetAnimBlueprintDebugData();
+			int32* IndexPtr = DebugData.NodePropertyToIndexMap.Find(this);
+
+			if (IndexPtr)
+			{
+				int32 AnimNodeIndex = *IndexPtr;
+				// reverse node index temporarily because of a bug in NodeGuidToIndexMap
+				AnimNodeIndex = AnimBlueprintClass->AnimNodeProperties.Num() - AnimNodeIndex - 1;
+
+				DebugNode = AnimBlueprintClass->AnimNodeProperties[AnimNodeIndex]->ContainerPtrToValuePtr<FAnimNode_Base>(PreviewSkelMeshComp->GetAnimInstance());
+			}
+		}
+	}
+
+	return DebugNode;
+}
+
+void UAnimGraphNode_Base::PinDefaultValueChanged(UEdGraphPin* Pin)
+{
+	Super::PinDefaultValueChanged(Pin);
+
+	CopyPinDefaultsToNodeData(Pin);
+
+	CastChecked<UAnimationGraph>(GetGraph())->OnPinDefaultValueChanged.Broadcast(Pin);
+}
\ No newline at end of file
diff --git a/Engine/Source/Editor/AnimGraph/Private/AnimGraphNode_Fabrik.cpp b/Engine/Source/Editor/AnimGraph/Private/AnimGraphNode_Fabrik.cpp
--- a/Engine/Source/Editor/AnimGraph/Private/AnimGraphNode_Fabrik.cpp
+++ b/Engine/Source/Editor/AnimGraph/Private/AnimGraphNode_Fabrik.cpp
@@ -3,6 +3,7 @@
 #include "AnimGraphPrivatePCH.h"
 #include "AnimGraphNode_Fabrik.h"
 #include "Animation/AnimInstance.h"
+#include "AnimNodeEditModes.h"
 
 /////////////////////////////////////////////////////
 // UAnimGraphNode_Fabrik 
@@ -35,60 +36,17 @@ FText UAnimGraphNode_Fabrik::GetNodeTitle(ENodeTitleType::Type TitleType) const
 	return GetControllerDescription();
 }
 
-FVector UAnimGraphNode_Fabrik::GetWidgetLocation(const USkeletalMeshComponent* SkelComp, FAnimNode_SkeletalControlBase* AnimNode)
+void UAnimGraphNode_Fabrik::CopyNodeDataToPreviewNode(FAnimNode_Base* InPreviewNode)
 {
-	FName& BoneName = Node.EffectorTransformBone.BoneName;
-	FVector Location = Node.EffectorTransform.GetLocation();
-	EBoneControlSpace Space = Node.EffectorTransformSpace;
-	FVector WidgetLoc = ConvertWidgetLocation(SkelComp, AnimNode->ForwardedPose, BoneName, Location, Space);
-	return WidgetLoc;
-}
-
-int32 UAnimGraphNode_Fabrik::GetWidgetMode(const USkeletalMeshComponent* SkelComp)
-{
-	uint32 BoneIndex = SkelComp->GetBoneIndex(Node.EffectorTransformBone.BoneName);
-
-	if (BoneIndex != INDEX_NONE)
-	{
-		return FWidget::WM_Translate;
-	}
-
-	return FWidget::WM_None;
-}
-
-FName UAnimGraphNode_Fabrik::FindSelectedBone()
-{
-	return Node.EffectorTransformBone.BoneName;
-}
-
-void UAnimGraphNode_Fabrik::DoTranslation(const USkeletalMeshComponent* SkelComp, FVector& Drag, FAnimNode_Base* InOutAnimNode)
-{
-	FAnimNode_Fabrik* Fabrik = static_cast<FAnimNode_Fabrik*>(InOutAnimNode);
-
-	if (Fabrik)
-	{
-		FVector Offset = ConvertCSVectorToBoneSpace(SkelComp, Drag, Fabrik->ForwardedPose, Node.EffectorTransformBone.BoneName, Node.EffectorTransformSpace);
-
-		Fabrik->EffectorTransform.AddToTranslation(Offset);
-		// copy same value to internal node for data consistency
-		Node.EffectorTransform.SetTranslation(Fabrik->EffectorTransform.GetTranslation());
-	}
-}
-
-void UAnimGraphNode_Fabrik::CopyNodeDataTo(FAnimNode_Base* AnimNode)
-{
-	FAnimNode_Fabrik* Fabrik = static_cast<FAnimNode_Fabrik*>(AnimNode);
+	FAnimNode_Fabrik* Fabrik = static_cast<FAnimNode_Fabrik*>(InPreviewNode);
 
 	// copies Pin values from the internal node to get data which are not compiled yet
 	Fabrik->EffectorTransform = Node.EffectorTransform;
 }
 
-void UAnimGraphNode_Fabrik::CopyNodeDataFrom(const FAnimNode_Base* InNewAnimNode)
+FEditorModeID UAnimGraphNode_Fabrik::GetEditorMode() const
 {
-	const FAnimNode_Fabrik* Fabrik = static_cast<const FAnimNode_Fabrik*>(InNewAnimNode);
-
-	// copies Pin data from updated values
-	Node.EffectorTransform = Fabrik->EffectorTransform;
+	return AnimNodeEditModes::Fabrik;
 }
 
 #undef LOCTEXT_NAMESPACE
\ No newline at end of file
diff --git a/Engine/Source/Editor/AnimGraph/Private/AnimGraphNode_ModifyBone.cpp b/Engine/Source/Editor/AnimGraph/Private/AnimGraphNode_ModifyBone.cpp
--- a/Engine/Source/Editor/AnimGraph/Private/AnimGraphNode_ModifyBone.cpp
+++ b/Engine/Source/Editor/AnimGraph/Private/AnimGraphNode_ModifyBone.cpp
@@ -3,6 +3,7 @@
 #include "AnimGraphPrivatePCH.h"
 #include "AnimGraphNode_ModifyBone.h"
 #include "CompilerResultsLog.h"
+#include "AnimNodeEditModes.h"
 
 /////////////////////////////////////////////////////
 // UAnimGraphNode_ModifyBone
@@ -79,230 +80,14 @@ FText UAnimGraphNode_ModifyBone::GetNodeTitle(ENodeTitleType::Type TitleType) co
 	return CachedNodeTitles[TitleType];
 }
 
-int32 UAnimGraphNode_ModifyBone::GetWidgetCoordinateSystem(const USkeletalMeshComponent* SkelComp)
+void UAnimGraphNode_ModifyBone::CopyNodeDataToPreviewNode(FAnimNode_Base* InPreviewNode)
 {
-	EBoneControlSpace Space = BCS_BoneSpace;
-	switch (CurWidgetMode)
-	{
-	case FWidget::WM_Rotate:
-		Space = Node.RotationSpace;
-		break;
-	case FWidget::WM_Translate:
-		Space = Node.TranslationSpace;
-		break;
-	case FWidget::WM_Scale:
-		Space = Node.ScaleSpace;
-		break;
-	}
-
-	switch (Space)
-	{
-	default:
-	case BCS_ParentBoneSpace:
-		//@TODO: No good way of handling this one
-		return COORD_World;
-	case BCS_BoneSpace:
-		return COORD_Local;
-	case BCS_ComponentSpace:
-	case BCS_WorldSpace:
-		return COORD_World;
-	}
-}
-
-FVector UAnimGraphNode_ModifyBone::GetWidgetLocation(const USkeletalMeshComponent* SkelComp, FAnimNode_SkeletalControlBase* AnimNode)
-{
-	USkeleton* Skeleton = SkelComp->SkeletalMesh->Skeleton;
-	FVector WidgetLoc = FVector::ZeroVector;
-
-	// if the current widget mode is translate, then shows the widget according to translation space
-	if (CurWidgetMode == FWidget::WM_Translate)
-	{
-		FCSPose<FCompactHeapPose>& MeshBases = AnimNode->ForwardedPose;
-		WidgetLoc = ConvertWidgetLocation(SkelComp, MeshBases, Node.BoneToModify.BoneName, GetNodeValue(FString("Translation"), Node.Translation), Node.TranslationSpace);
-
-		if (MeshBases.GetPose().IsValid() && Node.TranslationMode == BMM_Additive)
-		{
-			if(Node.TranslationSpace == EBoneControlSpace::BCS_WorldSpace ||
-				Node.TranslationSpace == EBoneControlSpace::BCS_ComponentSpace)
-			{
-				const FMeshPoseBoneIndex MeshBoneIndex(SkelComp->GetBoneIndex(Node.BoneToModify.BoneName));
-				const FCompactPoseBoneIndex BoneIndex = MeshBases.GetPose().GetBoneContainer().MakeCompactPoseIndex(MeshBoneIndex);
-
-				if (BoneIndex != INDEX_NONE)
-				{
-					const FTransform& BoneTM = MeshBases.GetComponentSpaceTransform(BoneIndex);
-					WidgetLoc += BoneTM.GetLocation();
-				}
-			}
-		}
-	}
-	else // if the current widget mode is not translate mode, then show the widget on the bone to modify
-	{
-		int32 MeshBoneIndex = SkelComp->GetBoneIndex(Node.BoneToModify.BoneName);
-
-		if (MeshBoneIndex != INDEX_NONE)
-		{
-			const FTransform BoneTM = SkelComp->GetBoneTransform(MeshBoneIndex);
-			WidgetLoc = BoneTM.GetLocation();
-		}
-	}
-	
-	return WidgetLoc;
-}
-
-EBoneModificationMode UAnimGraphNode_ModifyBone::GetBoneModificationMode(int32 InWidgetMode)
-{
-	FWidget::EWidgetMode InMode = (FWidget::EWidgetMode)InWidgetMode;
-	switch (InMode)
-	{
-	case FWidget::WM_Translate:
-		return Node.TranslationMode;
-		break;
-	case FWidget::WM_Rotate:
-		return Node.RotationMode;
-		break;
-	case FWidget::WM_Scale:
-		return Node.ScaleMode;
-		break;
-	case FWidget::WM_TranslateRotateZ:
-	case FWidget::WM_2D:
-		break;
-	}
-
-	return EBoneModificationMode::BMM_Ignore;
-}
-
-int32 UAnimGraphNode_ModifyBone::GetNextWidgetMode(int32 InWidgetMode)
-{
-	FWidget::EWidgetMode InMode = (FWidget::EWidgetMode)InWidgetMode;
-	switch (InMode)
-	{
-	case FWidget::WM_Translate:
-		return FWidget::WM_Rotate;
-	case FWidget::WM_Rotate:
-		return FWidget::WM_Scale;
-	case FWidget::WM_Scale:
-		return FWidget::WM_Translate;
-	case FWidget::WM_TranslateRotateZ:
-	case FWidget::WM_2D:
-		break;
-	}
-
-	return (int32)FWidget::WM_None;
-}
-
-int32 UAnimGraphNode_ModifyBone::FindValidWidgetMode(int32 InWidgetMode)
-{	
-	FWidget::EWidgetMode InMode = (FWidget::EWidgetMode)InWidgetMode;
-	FWidget::EWidgetMode ValidMode = InMode;
-	if (InMode == FWidget::WM_None)
-	{	// starts from Rotate mode
-		ValidMode = FWidget::WM_Rotate;
-	}
-
-	// find from current widget mode and loop 1 cycle until finding a valid mode
-	for (int32 Index = 0; Index < 3; Index++)
-	{
-		if (GetBoneModificationMode(ValidMode) != BMM_Ignore)
-		{
-			return ValidMode;
-		}
-
-		ValidMode = (FWidget::EWidgetMode)GetNextWidgetMode(ValidMode);
-	}
-	
-	// if couldn't find a valid mode, returns None
-	ValidMode = FWidget::WM_None;
-
-	return (int32)ValidMode;
-}
-
-int32 UAnimGraphNode_ModifyBone::GetWidgetMode(const USkeletalMeshComponent* SkelComp)
-{
-	int32 BoneIndex = SkelComp->GetBoneIndex(Node.BoneToModify.BoneName);
-	if (BoneIndex != INDEX_NONE)
-	{
-		CurWidgetMode = FindValidWidgetMode(CurWidgetMode);
-		return CurWidgetMode;
-	}
-
-	return (int32)FWidget::WM_None;
-}
-
-int32 UAnimGraphNode_ModifyBone::ChangeToNextWidgetMode(const USkeletalMeshComponent* SkelComp, int32 InCurWidgetMode)
-{
-	int32 NextWidgetMode = GetNextWidgetMode(InCurWidgetMode);
-	CurWidgetMode = FindValidWidgetMode(NextWidgetMode);
-
-	return CurWidgetMode;
-}
-
-bool UAnimGraphNode_ModifyBone::SetWidgetMode(const USkeletalMeshComponent* SkelComp, int32 InWidgetMode)
-{
-	// if InWidgetMode is available 
-	if (FindValidWidgetMode(InWidgetMode) == InWidgetMode)
-	{
-		CurWidgetMode = InWidgetMode;
-		return true;
-	}
-
-	return false;
-}
-
-FName UAnimGraphNode_ModifyBone::FindSelectedBone()
-{
-	return Node.BoneToModify.BoneName;
-}
-
-void UAnimGraphNode_ModifyBone::DoTranslation(const USkeletalMeshComponent* SkelComp, FVector& Drag, FAnimNode_Base* InOutAnimNode)
-{
-	FAnimNode_ModifyBone* ModifyBone = static_cast<FAnimNode_ModifyBone*>(InOutAnimNode);
-
-	if (ModifyBone->TranslationMode != EBoneModificationMode::BMM_Ignore)
-	{
-		FVector Offset = ConvertCSVectorToBoneSpace(SkelComp, Drag, ModifyBone->ForwardedPose, Node.BoneToModify.BoneName, Node.TranslationSpace);
-
-		ModifyBone->Translation += Offset;
-
-		Node.Translation = ModifyBone->Translation;
-	}
-}
-
-void UAnimGraphNode_ModifyBone::DoRotation(const USkeletalMeshComponent* SkelComp, FRotator& Rotation, FAnimNode_Base* InOutAnimNode)
-{
-	FAnimNode_ModifyBone* ModifyBone = static_cast<FAnimNode_ModifyBone*>(InOutAnimNode);
-
-	if (Node.RotationMode != EBoneModificationMode::BMM_Ignore)
-	{
-		FQuat DeltaQuat = ConvertCSRotationToBoneSpace(SkelComp, Rotation, ModifyBone->ForwardedPose, Node.BoneToModify.BoneName, Node.RotationSpace);
-
-		FQuat PrevQuat(ModifyBone->Rotation);
-		FQuat NewQuat = DeltaQuat * PrevQuat;
-		ModifyBone->Rotation = NewQuat.Rotator();
-		Node.Rotation = ModifyBone->Rotation;
-	}
-}
-
-void UAnimGraphNode_ModifyBone::DoScale(const USkeletalMeshComponent* SkelComp, FVector& Scale, FAnimNode_Base* InOutAnimNode)
-{
-	FAnimNode_ModifyBone* ModifyBone = static_cast<FAnimNode_ModifyBone*>(InOutAnimNode);
-
-	if (Node.ScaleMode != EBoneModificationMode::BMM_Ignore)
-	{
-		FVector Offset = Scale;
-		ModifyBone->Scale += Offset;
-		Node.Scale = ModifyBone->Scale;
-	}
-}
-
-void UAnimGraphNode_ModifyBone::CopyNodeDataTo(FAnimNode_Base* OutAnimNode)
-{
-	FAnimNode_ModifyBone* ModifyBone = static_cast<FAnimNode_ModifyBone*>(OutAnimNode);
+	FAnimNode_ModifyBone* ModifyBone = static_cast<FAnimNode_ModifyBone*>(InPreviewNode);
 
 	// copies Pin values from the internal node to get data which are not compiled yet
-	ModifyBone->Translation = GetNodeValue(FString("Translation"), Node.Translation);
-	ModifyBone->Rotation = GetNodeValue(FString("Rotation"), Node.Rotation);
-	ModifyBone->Scale = GetNodeValue(FString("Scale"), Node.Scale);
+	ModifyBone->Translation = Node.Translation;
+	ModifyBone->Rotation = Node.Rotation;
+	ModifyBone->Scale = Node.Scale;
 
 	// copies Modes
 	ModifyBone->TranslationMode = Node.TranslationMode;
@@ -315,21 +100,25 @@ void UAnimGraphNode_ModifyBone::CopyNodeDataTo(FAnimNode_Base* OutAnimNode)
 	ModifyBone->ScaleSpace = Node.ScaleSpace;
 }
 
-void UAnimGraphNode_ModifyBone::CopyNodeDataFrom(const FAnimNode_Base* InNewAnimNode)
+FEditorModeID UAnimGraphNode_ModifyBone::GetEditorMode() const
 {
-	const FAnimNode_ModifyBone* ModifyBone = static_cast<const FAnimNode_ModifyBone*>(InNewAnimNode);
+	return AnimNodeEditModes::ModifyBone;
+}
 
-	switch (CurWidgetMode)
+void UAnimGraphNode_ModifyBone::CopyPinDefaultsToNodeData(UEdGraphPin* InPin)
+{
+	if (InPin->GetName() == GET_MEMBER_NAME_STRING_CHECKED(FAnimNode_ModifyBone, Translation))
 	{
-	case FWidget::WM_Translate:
-		SetNodeValue(FString("Translation"), Node.Translation, ModifyBone->Translation);
-		break;
-	case FWidget::WM_Rotate:
-		SetNodeValue(FString("Rotation"), Node.Rotation, ModifyBone->Rotation);
-		break;
-	case FWidget::WM_Scale:
-		SetNodeValue(FString("Scale"), Node.Scale, ModifyBone->Scale);
-		break;
+		GetDefaultValue(GET_MEMBER_NAME_STRING_CHECKED(FAnimNode_ModifyBone, Translation), Node.Translation);
+	}
+	else if (InPin->GetName() == GET_MEMBER_NAME_STRING_CHECKED(FAnimNode_ModifyBone, Rotation))
+	{
+		GetDefaultValue(GET_MEMBER_NAME_STRING_CHECKED(FAnimNode_ModifyBone, Rotation), Node.Rotation);
+	}
+	else if (InPin->GetName() == GET_MEMBER_NAME_STRING_CHECKED(FAnimNode_ModifyBone, Scale))
+	{
+		GetDefaultValue(GET_MEMBER_NAME_STRING_CHECKED(FAnimNode_ModifyBone, Scale), Node.Scale);
 	}
 }
+
 #undef LOCTEXT_NAMESPACE
diff --git a/Engine/Source/Editor/AnimGraph/Private/AnimGraphNode_ObserveBone.cpp b/Engine/Source/Editor/AnimGraph/Private/AnimGraphNode_ObserveBone.cpp
--- a/Engine/Source/Editor/AnimGraph/Private/AnimGraphNode_ObserveBone.cpp
+++ b/Engine/Source/Editor/AnimGraph/Private/AnimGraphNode_ObserveBone.cpp
@@ -8,6 +8,7 @@
 #include "KismetDebugUtilities.h"
 #include "DetailLayoutBuilder.h"
 #include "PropertyHandle.h"
+#include "AnimNodeEditModes.h"
 
 #define LOCTEXT_NAMESPACE "ObserveBone"
 
@@ -136,48 +137,6 @@ FText UAnimGraphNode_ObserveBone::GetNodeTitle(ENodeTitleType::Type TitleType) c
 	}
 }
 
-int32 UAnimGraphNode_ObserveBone::GetWidgetCoordinateSystem(const USkeletalMeshComponent* SkelComp)
-{
-	switch (Node.DisplaySpace)
-	{
-	default:
-	case BCS_ParentBoneSpace:
-		//@TODO: No good way of handling this one
-		return COORD_World;
-	case BCS_BoneSpace:
-		return COORD_Local;
-	case BCS_ComponentSpace:
-	case BCS_WorldSpace:
-		return COORD_World;
-	}
-}
-
-FVector UAnimGraphNode_ObserveBone::GetWidgetLocation(const USkeletalMeshComponent* SkelComp, FAnimNode_SkeletalControlBase* AnimNode)
-{
-	USkeleton* Skeleton = SkelComp->SkeletalMesh->Skeleton;
-	FVector WidgetLoc = FVector::ZeroVector;
-
-	int32 MeshBoneIndex = SkelComp->GetBoneIndex(Node.BoneToObserve.BoneName);
-
-	if (MeshBoneIndex != INDEX_NONE)
-	{
-		const FTransform BoneTM = SkelComp->GetBoneTransform(MeshBoneIndex);
-		WidgetLoc = BoneTM.GetLocation();
-	}
-	
-	return WidgetLoc;
-}
-
-int32 UAnimGraphNode_ObserveBone::GetWidgetMode(const USkeletalMeshComponent* SkelComp)
-{
-	return (int32)FWidget::WM_Translate;
-}
-
-FName UAnimGraphNode_ObserveBone::FindSelectedBone()
-{
-	return Node.BoneToObserve.BoneName;
-}
-
 FLinearColor UAnimGraphNode_ObserveBone::GetNodeTitleColor() const
 {
 	return FLinearColor(0.7f, 0.7f, 0.7f);
@@ -198,4 +157,9 @@ void UAnimGraphNode_ObserveBone::CustomizeDetails(IDetailLayoutBuilder& DetailBu
 	DetailBuilder.HideProperty(NodeHandle->GetChildHandle(GET_MEMBER_NAME_CHECKED(FAnimNode_SkeletalControlBase, AlphaScaleBias)));
 }
 
+FEditorModeID UAnimGraphNode_ObserveBone::GetEditorMode() const
+{
+	return AnimNodeEditModes::ObserveBone;
+}
+
 #undef LOCTEXT_NAMESPACE
diff --git a/Engine/Source/Editor/AnimGraph/Private/AnimGraphNode_PoseBlendNode.cpp b/Engine/Source/Editor/AnimGraph/Private/AnimGraphNode_PoseBlendNode.cpp
--- a/Engine/Source/Editor/AnimGraph/Private/AnimGraphNode_PoseBlendNode.cpp
+++ b/Engine/Source/Editor/AnimGraph/Private/AnimGraphNode_PoseBlendNode.cpp
@@ -44,14 +44,6 @@ UAnimGraphNode_PoseBlendNode::UAnimGraphNode_PoseBlendNode(const FObjectInitiali
 {
 }
 
-void UAnimGraphNode_PoseBlendNode::PreloadRequiredAssets()
-{
-	PreloadObject(Node.PoseAsset);
-
-	Super::PreloadRequiredAssets();
-}
-
-
 void UAnimGraphNode_PoseBlendNode::GetAllAnimationSequencesReferred(TArray<UAnimationAsset*>& AnimationAssets) const
 {
 	if(Node.PoseAsset)
@@ -76,10 +68,7 @@ FText UAnimGraphNode_PoseBlendNode::GetNodeTitleForPoseAsset(ENodeTitleType::Typ
 	FFormatNamedArguments Args;
 	Args.Add(TEXT("PoseAssetName"), FText::FromString(InPoseAsset->GetName()));
 
-	// FText::Format() is slow, so we cache this to save on performance
-	CachedNodeTitle.SetCachedText(FText::Format(LOCTEXT("PoseByName_Title", "{PoseAssetName}"), Args), this);
-
-	return CachedNodeTitle;
+	return FText::Format(LOCTEXT("PoseByName_Title", "{PoseAssetName}"), Args);
 }
 
 FText UAnimGraphNode_PoseBlendNode::GetNodeTitle(ENodeTitleType::Type TitleType) const
@@ -101,15 +90,10 @@ FText UAnimGraphNode_PoseBlendNode::GetNodeTitle(ENodeTitleType::Type TitleType)
 			return LOCTEXT("PoseByName_TitleNONE", "Pose (None)");
 		}
 	}
-	// @TODO: don't know enough about this node type to comfortably assert that
-	//        the CacheName won't change after the node has spawned... until
-	//        then, we'll leave this optimization off
-	else //if (CachedNodeTitle.IsOutOfDate(this))
+	else
 	{
-		GetNodeTitleForPoseAsset(TitleType, Node.PoseAsset);
+		return GetNodeTitleForPoseAsset(TitleType, Node.PoseAsset);
 	}
-
-	return CachedNodeTitle;
 }
 
 void UAnimGraphNode_PoseBlendNode::GetMenuActions(FBlueprintActionDatabaseRegistrar& ActionRegistrar) const
@@ -254,58 +238,11 @@ FText UAnimGraphNode_PoseBlendNode::GetMenuCategory() const
 	return LOCTEXT("PoseAssetCategory_Label", "Poses");
 }
 
-void UAnimGraphNode_PoseBlendNode::SetAnimationAsset(UAnimationAsset* Asset)
-{
-	if (UPoseAsset* PoseAsset =  Cast<UPoseAsset>(Asset))
-	{
-		Node.PoseAsset = PoseAsset;
-	}
-}
-
-void UAnimGraphNode_PoseBlendNode::ValidateAnimNodeDuringCompilation(class USkeleton* ForSkeleton, class FCompilerResultsLog& MessageLog)
-{
-	UPoseAsset* PoseAssetToCheck = Node.PoseAsset;
-	UEdGraphPin* PoseAssetPin = FindPin(GET_MEMBER_NAME_STRING_CHECKED(FAnimNode_PoseBlendNode, PoseAsset));
-	if (PoseAssetPin != nullptr && PoseAssetToCheck == nullptr)
-	{
-		PoseAssetToCheck = Cast<UPoseAsset>(PoseAssetPin->DefaultObject);
-	}
-
-	if (PoseAssetToCheck == nullptr)
-	{
-		if (PoseAssetPin == nullptr || PoseAssetPin->LinkedTo.Num() == 0)
-		{
-			MessageLog.Error(TEXT("@@ references an unknown poseasset"), this);
-		}
-	}
-	else
-	{
-		USkeleton* SeqSkeleton = PoseAssetToCheck->GetSkeleton();
-		if (SeqSkeleton&& // if PoseAsset doesn't have skeleton, it might be due to PoseAsset not loaded yet, @todo: wait with anim blueprint compilation until all assets are loaded?
-			!SeqSkeleton->IsCompatible(ForSkeleton))
-		{
-			MessageLog.Error(TEXT("@@ references poseasset that uses different skeleton @@"), this, SeqSkeleton);
-		}
-	}
-}
-
 bool UAnimGraphNode_PoseBlendNode::DoesSupportTimeForTransitionGetter() const
 {
 	return false;
 }
 
-UAnimationAsset* UAnimGraphNode_PoseBlendNode::GetAnimationAsset() const 
-{
-	UPoseAsset* PoseAsset = Node.PoseAsset;
-	UEdGraphPin* PoseAssetPin = FindPin(GET_MEMBER_NAME_STRING_CHECKED(FAnimNode_PoseBlendNode, PoseAsset));
-	if (PoseAssetPin != nullptr && PoseAsset == nullptr)
-	{
-		PoseAsset = Cast<UPoseAsset>(PoseAssetPin->DefaultObject);
-	}
-
-	return PoseAsset;
-}
-
 void UAnimGraphNode_PoseBlendNode::GetContextMenuActions(const FGraphNodeContextMenuBuilder& Context) const
 {
 	if (!Context.bIsDebugging)
diff --git a/Engine/Source/Editor/AnimGraph/Private/AnimGraphNode_PoseDriver.cpp b/Engine/Source/Editor/AnimGraph/Private/AnimGraphNode_PoseDriver.cpp
--- a/Engine/Source/Editor/AnimGraph/Private/AnimGraphNode_PoseDriver.cpp
+++ b/Engine/Source/Editor/AnimGraph/Private/AnimGraphNode_PoseDriver.cpp
@@ -3,10 +3,7 @@
 #include "AnimGraphPrivatePCH.h"
 #include "AnimGraphNode_PoseDriver.h"
 #include "CompilerResultsLog.h"
-#include "PropertyEditing.h"
-#include "AnimationCustomVersion.h"
-#include "ContentBrowserModule.h"
-#include "Animation/PoseAsset.h"
+#include "AnimNodeEditModes.h"
 
 #define LOCTEXT_NAMESPACE "A3Nodes"
 
@@ -22,68 +19,13 @@ FText UAnimGraphNode_PoseDriver::GetTooltipText() const
 }
 
 FText UAnimGraphNode_PoseDriver::GetNodeTitle(ENodeTitleType::Type TitleType) const
-{
-	return GetControllerDescription();
-}
-
-FText UAnimGraphNode_PoseDriver::GetControllerDescription() const
 {
 	return LOCTEXT("PoseDriver", "Pose Driver");
 }
 
-void UAnimGraphNode_PoseDriver::Draw(FPrimitiveDrawInterface* PDI, USkeletalMeshComponent* SkelMeshComp) const
-{
-
-}
-
-void UAnimGraphNode_PoseDriver::DrawCanvas(FViewport& InViewport, FSceneView& View, FCanvas& Canvas, USkeletalMeshComponent * PreviewSkelMeshComp) const
+FText UAnimGraphNode_PoseDriver::GetMenuCategory() const
 {
-	// Get the node we are debugging
-	FAnimNode_PoseDriver* DebugNode = (FAnimNode_PoseDriver*)FindDebugAnimNode(PreviewSkelMeshComp);
-
-	if (DebugNode == nullptr)
-	{
-		return;
-	}
-
-	const int MarginX = 60;
-	const int XDelta = 10;
-
-	int PosY = 60;
-	const int YDelta = 20;
-
-	// Display bone we are watching
-	FText InterpItemText = FText::Format(LOCTEXT("BoneFormat", "Bone: {0}"), FText::FromName(DebugNode->SourceBone.BoneName));
-	FCanvasTextItem InterpItem(FVector2D(MarginX, PosY), InterpItemText, GEngine->GetSmallFont(), FLinearColor::White);
-	Canvas.DrawItem(InterpItem);
-	PosY += YDelta;
-
-	int32 PoseIndex = 0;
-	for (const FPoseDriverPoseInfo& PoseInfo : DebugNode->PoseInfos)
-	{
-		FString PoseItemString = FString::Printf(TEXT("%s : %f"), *PoseInfo.PoseName.ToString(), PoseInfo.PoseWeight);
-		FCanvasTextItem PoseItem(FVector2D(MarginX + XDelta, PosY), FText::FromString(PoseItemString), GEngine->GetSmallFont(), FLinearColor::White);
-		Canvas.DrawItem(PoseItem);
-		PosY += YDelta;
-
-		PoseIndex++;
-	}
-
-	PosY += YDelta;
-
-	// Parameters header
-	FCanvasTextItem ParamsHeaderItem(FVector2D(MarginX, PosY), LOCTEXT("Parameters", "Parameters"), GEngine->GetSmallFont(), FLinearColor::White);
-	Canvas.DrawItem(ParamsHeaderItem);
-	PosY += YDelta;
-
-	// Iterate over Parameters
-	for (const FPoseDriverParam& Param : DebugNode->ResultParamSet.Params)
-	{
-		FString ParamItemString = FString::Printf(TEXT("%s : %f"), *Param.ParamInfo.Name.ToString(), Param.ParamValue);
-		FCanvasTextItem ParamItem(FVector2D(MarginX + XDelta, PosY), FText::FromString(ParamItemString), GEngine->GetSmallFont(), FLinearColor::White);
-		Canvas.DrawItem(ParamItem);
-		PosY += YDelta;
-	}
+	return LOCTEXT("PoseAssetCategory_Label", "Poses");
 }
 
 
@@ -97,10 +39,9 @@ void UAnimGraphNode_PoseDriver::ValidateAnimNodeDuringCompilation(USkeleton* For
 	Super::ValidateAnimNodeDuringCompilation(ForSkeleton, MessageLog);
 }
 
-void UAnimGraphNode_PoseDriver::CustomizeDetails(IDetailLayoutBuilder& DetailBuilder)
+FEditorModeID UAnimGraphNode_PoseDriver::GetEditorMode() const
 {
-
+	return AnimNodeEditModes::PoseDriver;
 }
 
-
 #undef LOCTEXT_NAMESPACE
diff --git a/Engine/Source/Editor/AnimGraph/Private/AnimGraphNode_PoseHandler.cpp b/Engine/Source/Editor/AnimGraph/Private/AnimGraphNode_PoseHandler.cpp
new file mode 100644
--- /dev/null
+++ b/Engine/Source/Editor/AnimGraph/Private/AnimGraphNode_PoseHandler.cpp
@@ -0,0 +1,69 @@
+// Copyright 1998-2016 Epic Games, Inc. All Rights Reserved.
+
+#include "AnimGraphPrivatePCH.h"
+#include "AnimGraphNode_PoseHandler.h"
+#include "AnimNodes/AnimNode_PoseHandler.h"
+#include "CompilerResultsLog.h"
+#include "Animation/PoseAsset.h"
+
+UAnimGraphNode_PoseHandler::UAnimGraphNode_PoseHandler(const FObjectInitializer& ObjectInitializer)
+	:Super(ObjectInitializer)
+{
+
+}
+
+void UAnimGraphNode_PoseHandler::ValidateAnimNodeDuringCompilation(USkeleton* ForSkeleton, FCompilerResultsLog& MessageLog)
+{
+	UPoseAsset* PoseAssetToCheck = GetPoseHandlerNode()->PoseAsset;
+	UEdGraphPin* PoseAssetPin = FindPin(GET_MEMBER_NAME_STRING_CHECKED(FAnimNode_PoseHandler, PoseAsset));
+	if (PoseAssetPin != nullptr && PoseAssetToCheck == nullptr)
+	{
+		PoseAssetToCheck = Cast<UPoseAsset>(PoseAssetPin->DefaultObject);
+	}
+
+	if (PoseAssetToCheck == nullptr)
+	{
+		if (PoseAssetPin == nullptr || PoseAssetPin->LinkedTo.Num() == 0)
+		{
+			MessageLog.Error(TEXT("@@ references an unknown poseasset"), this);
+		}
+	}
+	else
+	{
+		USkeleton* SeqSkeleton = PoseAssetToCheck->GetSkeleton();
+		if (SeqSkeleton && // if PoseAsset doesn't have skeleton, it might be due to PoseAsset not loaded yet, @todo: wait with anim blueprint compilation until all assets are loaded?
+			!SeqSkeleton->IsCompatible(ForSkeleton))
+		{
+			MessageLog.Error(TEXT("@@ references poseasset that uses different skeleton @@"), this, SeqSkeleton);
+		}
+	}
+
+	Super::ValidateAnimNodeDuringCompilation(ForSkeleton, MessageLog);
+}
+
+void UAnimGraphNode_PoseHandler::SetAnimationAsset(UAnimationAsset* Asset)
+{
+	if (UPoseAsset* PoseAsset = Cast<UPoseAsset>(Asset))
+	{
+		GetPoseHandlerNode()->PoseAsset = PoseAsset;
+	}
+}
+
+void UAnimGraphNode_PoseHandler::PreloadRequiredAssets()
+{
+	PreloadObject(GetPoseHandlerNode()->PoseAsset);
+
+	Super::PreloadRequiredAssets();
+}
+
+UAnimationAsset* UAnimGraphNode_PoseHandler::GetAnimationAsset() const
+{
+	UPoseAsset* PoseAsset = GetPoseHandlerNode()->PoseAsset;
+	UEdGraphPin* PoseAssetPin = FindPin(GET_MEMBER_NAME_STRING_CHECKED(FAnimNode_PoseHandler, PoseAsset));
+	if (PoseAssetPin != nullptr && PoseAsset == nullptr)
+	{
+		PoseAsset = Cast<UPoseAsset>(PoseAssetPin->DefaultObject);
+	}
+
+	return PoseAsset;
+}
diff --git a/Engine/Source/Editor/AnimGraph/Private/AnimGraphNode_SkeletalControlBase.cpp b/Engine/Source/Editor/AnimGraph/Private/AnimGraphNode_SkeletalControlBase.cpp
--- a/Engine/Source/Editor/AnimGraph/Private/AnimGraphNode_SkeletalControlBase.cpp
+++ b/Engine/Source/Editor/AnimGraph/Private/AnimGraphNode_SkeletalControlBase.cpp
@@ -7,6 +7,7 @@
 #include "Animation/AnimInstance.h"
 #include "CompilerResultsLog.h"
 #include "AnimNode_SkeletalControlBase.h"
+
 /////////////////////////////////////////////////////
 // UAnimGraphNode_SkeletalControlBase
 
@@ -19,6 +20,7 @@ UAnimGraphNode_SkeletalControlBase::UAnimGraphNode_SkeletalControlBase(const FOb
 
 int32 UAnimGraphNode_SkeletalControlBase::GetWidgetCoordinateSystem(const USkeletalMeshComponent* SkelComp)
 {
+	PRAGMA_DISABLE_DEPRECATION_WARNINGS
 	if (GetWidgetMode(SkelComp) == FWidget::WM_Scale)
 	{
 		return COORD_Local;
@@ -27,6 +29,7 @@ int32 UAnimGraphNode_SkeletalControlBase::GetWidgetCoordinateSystem(const USkele
 	{
 		return COORD_World;
 	}
+	PRAGMA_ENABLE_DEPRECATION_WARNINGS
 }
 
 // returns int32 instead of EWidgetMode because of compiling issue on Mac
@@ -37,7 +40,9 @@ int32 UAnimGraphNode_SkeletalControlBase::GetWidgetMode(const USkeletalMeshCompo
 
 int32 UAnimGraphNode_SkeletalControlBase::ChangeToNextWidgetMode(const USkeletalMeshComponent* SkelComp, int32 CurWidgetMode)
 {
+	PRAGMA_DISABLE_DEPRECATION_WARNINGS
 	return GetWidgetMode(SkelComp);
+	PRAGMA_ENABLE_DEPRECATION_WARNINGS
 }
 
 FName UAnimGraphNode_SkeletalControlBase::FindSelectedBone()
@@ -71,33 +76,6 @@ void UAnimGraphNode_SkeletalControlBase::CreateOutputPins()
 	CreatePin(EGPD_Output, Schema->PC_Struct, TEXT(""), FComponentSpacePoseLink::StaticStruct(), /*bIsArray=*/ false, /*bIsReference=*/ false, TEXT("Pose"));
 }
 
-FAnimNode_SkeletalControlBase* UAnimGraphNode_SkeletalControlBase::FindDebugAnimNode(USkeletalMeshComponent * PreviewSkelMeshComp) const
-{
-	FAnimNode_SkeletalControlBase* DebugNode = nullptr;
-
-	if (PreviewSkelMeshComp != nullptr && PreviewSkelMeshComp->GetAnimInstance() != nullptr)
-	{
-		// find an anim node index from debug data
-		UAnimBlueprintGeneratedClass* AnimBlueprintClass = Cast<UAnimBlueprintGeneratedClass>(PreviewSkelMeshComp->GetAnimInstance()->GetClass());
-		if (AnimBlueprintClass)
-		{
-			FAnimBlueprintDebugData& DebugData = AnimBlueprintClass->GetAnimBlueprintDebugData();
-			int32* IndexPtr = DebugData.NodePropertyToIndexMap.Find(this);
-
-			if (IndexPtr)
-			{
-				int32 AnimNodeIndex = *IndexPtr;
-				// reverse node index temporarily because of a bug in NodeGuidToIndexMap
-				AnimNodeIndex = AnimBlueprintClass->AnimNodeProperties.Num() - AnimNodeIndex - 1;
-
-				DebugNode = AnimBlueprintClass->AnimNodeProperties[AnimNodeIndex]->ContainerPtrToValuePtr<FAnimNode_SkeletalControlBase>(PreviewSkelMeshComp->GetAnimInstance());
-			}
-		}
-	}
-
-	return DebugNode;
-}
-
 
 void UAnimGraphNode_SkeletalControlBase::ConvertToComponentSpaceTransform(const USkeletalMeshComponent* SkelComp, const FTransform & InTransform, FTransform & OutCSTransform, int32 BoneIndex, EBoneControlSpace Space) const
 {
@@ -390,4 +368,5 @@ void UAnimGraphNode_SkeletalControlBase::ValidateAnimNodePostCompile(FCompilerRe
 		}
 	}
 }
+
 #undef LOCTEXT_NAMESPACE
\ No newline at end of file
diff --git a/Engine/Source/Editor/AnimGraph/Private/AnimGraphNode_TwoBoneIK.cpp b/Engine/Source/Editor/AnimGraph/Private/AnimGraphNode_TwoBoneIK.cpp
--- a/Engine/Source/Editor/AnimGraph/Private/AnimGraphNode_TwoBoneIK.cpp
+++ b/Engine/Source/Editor/AnimGraph/Private/AnimGraphNode_TwoBoneIK.cpp
@@ -3,6 +3,7 @@
 #include "AnimGraphPrivatePCH.h"
 #include "AnimGraphNode_TwoBoneIK.h"
 #include "DebugRenderSceneProxy.h"
+#include "AnimNodeEditModes.h"
 
 // for customization details
 #include "../../PropertyEditor/Public/PropertyHandle.h"
@@ -37,7 +38,6 @@ TSharedPtr<FTwoBoneIKDelegate> UAnimGraphNode_TwoBoneIK::TwoBoneIKDelegate = NUL
 UAnimGraphNode_TwoBoneIK::UAnimGraphNode_TwoBoneIK(const FObjectInitializer& ObjectInitializer)
 	: Super(ObjectInitializer)
 {
-	BoneSelectMode = BSM_EndEffector;
 }
 
 FText UAnimGraphNode_TwoBoneIK::GetControllerDescription() const
@@ -77,245 +77,24 @@ FText UAnimGraphNode_TwoBoneIK::GetNodeTitle(ENodeTitleType::Type TitleType) con
 	return CachedNodeTitles[TitleType];
 }
 
-void UAnimGraphNode_TwoBoneIK::Draw( FPrimitiveDrawInterface* PDI, USkeletalMeshComponent* SkelMeshComp ) const
+void UAnimGraphNode_TwoBoneIK::CopyNodeDataToPreviewNode(FAnimNode_Base* InPreviewNode)
 {
-	if (SkelMeshComp && SkelMeshComp->SkeletalMesh && SkelMeshComp->SkeletalMesh->Skeleton)
-	{
-		USkeleton * Skeleton = SkelMeshComp->SkeletalMesh->Skeleton;
-
-		DrawTargetLocation(PDI, SkelMeshComp, Skeleton, Node.EffectorLocationSpace, Node.EffectorSpaceBoneName, Node.EffectorLocation, FColor(255, 128, 128), FColor(180, 128, 128));
-		DrawTargetLocation(PDI, SkelMeshComp, Skeleton, Node.JointTargetLocationSpace, Node.JointTargetSpaceBoneName, Node.JointTargetLocation, FColor(128, 255, 128), FColor(128, 180, 128));
-	}
-}
-
-void UAnimGraphNode_TwoBoneIK::DrawTargetLocation(FPrimitiveDrawInterface* PDI, USkeletalMeshComponent* SkelComp, USkeleton * Skeleton, EBoneControlSpace SpaceBase, FName SpaceBoneName, const FVector & TargetLocation, const FColor & TargetColor, const FColor & BoneColor) const
-{
-	const bool bInBoneSpace = (SpaceBase == BCS_ParentBoneSpace) || (SpaceBase == BCS_BoneSpace);
-	const int32 SpaceBoneIndex = bInBoneSpace ? Skeleton->GetReferenceSkeleton().FindBoneIndex(SpaceBoneName) : INDEX_NONE;
-	// Transform EffectorLocation from EffectorLocationSpace to ComponentSpace.
-	FTransform TargetTransform = FTransform (TargetLocation);
-	FTransform CSTransform;
-
-	ConvertToComponentSpaceTransform(SkelComp, TargetTransform, CSTransform, SpaceBoneIndex, SpaceBase);
-
-	FTransform WorldTransform = CSTransform * SkelComp->ComponentToWorld;
-
-#if 0 // @TODO : remove this code because this doesn't show correct target location
-	DrawCoordinateSystem( PDI, WorldTransform.GetLocation(), WorldTransform.GetRotation().Rotator(), 20.f, SDPG_Foreground );
-	DrawWireDiamond( PDI, WorldTransform.ToMatrixWithScale(), 2.f, TargetColor, SDPG_Foreground );
-#endif
-
-	if (bInBoneSpace)
-	{
-		ConvertToComponentSpaceTransform(SkelComp, FTransform::Identity, CSTransform, SpaceBoneIndex, SpaceBase);
-		WorldTransform = CSTransform * SkelComp->ComponentToWorld;
-		DrawCoordinateSystem( PDI, WorldTransform.GetLocation(), WorldTransform.GetRotation().Rotator(), 20.f, SDPG_Foreground );
-		DrawWireDiamond( PDI, WorldTransform.ToMatrixWithScale(), 2.f, BoneColor, SDPG_Foreground );
-	}
-}
-
-void UAnimGraphNode_TwoBoneIK::CopyNodeDataTo(FAnimNode_Base* AnimNode)
-{
-	FAnimNode_TwoBoneIK* TwoBoneIK = static_cast<FAnimNode_TwoBoneIK*>(AnimNode);
+	FAnimNode_TwoBoneIK* TwoBoneIK = static_cast<FAnimNode_TwoBoneIK*>(InPreviewNode);
 
 	// copies Pin values from the internal node to get data which are not compiled yet
-	TwoBoneIK->EffectorLocation = GetNodeValue(FString("EffectorLocation"), Node.EffectorLocation);
-	TwoBoneIK->JointTargetLocation = GetNodeValue(FString("JointTargetLocation"), Node.JointTargetLocation);
-}
-
-void UAnimGraphNode_TwoBoneIK::CopyNodeDataFrom(const FAnimNode_Base* InNewAnimNode)
-{
-	const FAnimNode_TwoBoneIK* TwoBoneIK = static_cast<const FAnimNode_TwoBoneIK*>(InNewAnimNode);
-	
-	if (BoneSelectMode == BSM_EndEffector)
-	{
-		SetNodeValue(FString("EffectorLocation"), Node.EffectorLocation, TwoBoneIK->EffectorLocation);
-	}
-	else
-	{
-		SetNodeValue(FString("JointTargetLocation"), Node.JointTargetLocation, TwoBoneIK->JointTargetLocation);
-	}
-}
-
-
-FVector UAnimGraphNode_TwoBoneIK::GetWidgetLocation(const USkeletalMeshComponent* SkelComp, FAnimNode_SkeletalControlBase* AnimNode)
-{
-	EBoneControlSpace Space;
-	FName BoneName;
-	FVector Location;
-	 
-	if (BoneSelectMode == BSM_EndEffector)
-	{
-		Space = Node.EffectorLocationSpace;
-		Location = GetNodeValue(FString("EffectorLocation"), Node.EffectorLocation);
-		BoneName = Node.EffectorSpaceBoneName;
-
-	}
-	else // BSM_JointTarget
-	{
-		Space = Node.JointTargetLocationSpace;
-
-		if (Space == BCS_WorldSpace || Space == BCS_ComponentSpace)
-		{
-			return FVector::ZeroVector;
-		}
-		Location = GetNodeValue(FString("JointTargetLocation"), Node.JointTargetLocation);
-		BoneName = Node.JointTargetSpaceBoneName;
-	}
-
-	return ConvertWidgetLocation(SkelComp, AnimNode->ForwardedPose, BoneName, Location, Space);
-}
-
-int32 UAnimGraphNode_TwoBoneIK::GetWidgetMode(const USkeletalMeshComponent* SkelComp)
-{
-	int32 BoneIndex = SkelComp->GetBoneIndex(Node.IKBone.BoneName);
-	// Two bone IK node uses only Translate
-	if (BoneIndex != INDEX_NONE)
-	{
-		return FWidget::WM_Translate;
-	}
-
-	return FWidget::WM_None;
-}
-
-void UAnimGraphNode_TwoBoneIK::MoveSelectActorLocation(const USkeletalMeshComponent* SkelComp, FAnimNode_SkeletalControlBase* AnimNode)
-{
-	USkeleton * Skeleton = SkelComp->SkeletalMesh->Skeleton;
-
-	// create a bone select actor
-
-	if (!BoneSelectActor.IsValid())
-	{
-		if (Node.JointTargetLocationSpace == EBoneControlSpace::BCS_BoneSpace ||
-			Node.JointTargetLocationSpace == EBoneControlSpace::BCS_ParentBoneSpace)
-		{
-			int32 JointTargetIndex = Skeleton->GetReferenceSkeleton().FindBoneIndex(Node.JointTargetSpaceBoneName);
-
-			if (JointTargetIndex != INDEX_NONE)
-			{
-				UWorld* World = SkelComp->GetWorld();
-				check(World);
-
-				BoneSelectActor = World->SpawnActor<ABoneSelectActor>(FVector(0), FRotator(0, 0, 0));
-				check(BoneSelectActor.IsValid());
-			}
-		}
-	}
-
-	if (!BoneSelectActor.IsValid())
-	{
-		return;
-	}
-
-	// move the actor's position
-	if (BoneSelectMode == BSM_JointTarget)
-	{
-		FName BoneName;
-		int32 EffectorBoneIndex = SkelComp->GetBoneIndex(Node.EffectorSpaceBoneName);
-		if (EffectorBoneIndex != INDEX_NONE)
-		{
-			BoneName = Node.EffectorSpaceBoneName;
-		}
-		else
-		{
-			BoneName = Node.IKBone.BoneName;
-		}
-
-		FVector ActorLocation = ConvertWidgetLocation(SkelComp, AnimNode->ForwardedPose, BoneName, Node.EffectorLocation, Node.EffectorLocationSpace);
-		BoneSelectActor->SetActorLocation(ActorLocation + FVector(0, 10, 0));
-
-	}
-	else if (BoneSelectMode == BSM_EndEffector)
-	{
-		int32 JointTargetIndex = SkelComp->GetBoneIndex(Node.JointTargetSpaceBoneName);
-
-		if (JointTargetIndex != INDEX_NONE)
-		{
-			FVector ActorLocation = ConvertWidgetLocation(SkelComp, AnimNode->ForwardedPose, Node.JointTargetSpaceBoneName, Node.JointTargetLocation, Node.JointTargetLocationSpace);
-			BoneSelectActor->SetActorLocation(ActorLocation + FVector(0, 10, 0));
-		}
-	}
-}
-
-FName UAnimGraphNode_TwoBoneIK::FindSelectedBone()
-{
-	FName SelectedBone;
-
-	// should return mesh bone index
-	if (BoneSelectMode == BSM_EndEffector)
-	{
-		if (Node.EffectorLocationSpace == EBoneControlSpace::BCS_BoneSpace ||
-			Node.EffectorLocationSpace == EBoneControlSpace::BCS_ParentBoneSpace)
-		{
-			SelectedBone = Node.EffectorSpaceBoneName;
-		}
-		else
-		{
-			SelectedBone = Node.IKBone.BoneName;
-		}
-
-	}
-	else if (BoneSelectMode == BSM_JointTarget)
-	{
-		SelectedBone = Node.JointTargetSpaceBoneName;
-	}
-
-	return SelectedBone;
-}
-
-bool UAnimGraphNode_TwoBoneIK::IsActorClicked(HActor* ActorHitProxy)
-{
-	ABoneSelectActor* BoneActor = Cast<ABoneSelectActor>(ActorHitProxy->Actor);
-
-	if (BoneActor)
-	{
-		return true;
-	}
-
-	return false;
-}
-
-void UAnimGraphNode_TwoBoneIK::ProcessActorClick(HActor* ActorHitProxy)
-{
-	// toggle bone selection mode
-	if (BoneSelectMode == BSM_EndEffector)
-	{
-		BoneSelectMode = BSM_JointTarget;
-	}
-	else
-	{
-		BoneSelectMode = BSM_EndEffector;
-	}
+	TwoBoneIK->EffectorLocation = Node.EffectorLocation;
+	TwoBoneIK->JointTargetLocation = Node.JointTargetLocation;
 }
 
-void UAnimGraphNode_TwoBoneIK::DoTranslation(const USkeletalMeshComponent* SkelComp, FVector& Drag, FAnimNode_Base* InOutAnimNode)
+void UAnimGraphNode_TwoBoneIK::CopyPinDefaultsToNodeData(UEdGraphPin* InPin)
 {
-	FAnimNode_TwoBoneIK* TwoBoneIK = static_cast<FAnimNode_TwoBoneIK*>(InOutAnimNode);
-
-	if (BoneSelectMode == BSM_EndEffector)
+	if (InPin->GetName() == GET_MEMBER_NAME_STRING_CHECKED(FAnimNode_TwoBoneIK, EffectorLocation))
 	{
-		FVector Offset = ConvertCSVectorToBoneSpace(SkelComp, Drag, TwoBoneIK->ForwardedPose, FindSelectedBone(), Node.EffectorLocationSpace);
-
-		TwoBoneIK->EffectorLocation += Offset;
-		Node.EffectorLocation = TwoBoneIK->EffectorLocation;
+		GetDefaultValue(GET_MEMBER_NAME_STRING_CHECKED(FAnimNode_TwoBoneIK, EffectorLocation), Node.EffectorLocation);
 	}
-	else if (BoneSelectMode == BSM_JointTarget)
-	{
-		FVector Offset = ConvertCSVectorToBoneSpace(SkelComp, Drag, TwoBoneIK->ForwardedPose, FindSelectedBone(), Node.JointTargetLocationSpace);
-
-		TwoBoneIK->JointTargetLocation += Offset;
-		Node.JointTargetLocation = TwoBoneIK->JointTargetLocation;
-	}
-}
-
-void UAnimGraphNode_TwoBoneIK::DeselectActor(USkeletalMeshComponent* SkelComp)
-{
-	if(BoneSelectActor!=NULL)
+	else if (InPin->GetName() == GET_MEMBER_NAME_STRING_CHECKED(FAnimNode_TwoBoneIK, JointTargetLocation))
 	{
-		if (SkelComp->GetWorld()->DestroyActor(BoneSelectActor.Get()))
-		{
-			BoneSelectActor = NULL;
-		}
+		GetDefaultValue(GET_MEMBER_NAME_STRING_CHECKED(FAnimNode_TwoBoneIK, JointTargetLocation), Node.JointTargetLocation);
 	}
 }
 
@@ -397,4 +176,9 @@ void UAnimGraphNode_TwoBoneIK::CustomizeDetails(class IDetailLayoutBuilder& Deta
 	}
 }
 
+FEditorModeID UAnimGraphNode_TwoBoneIK::GetEditorMode() const
+{
+	return AnimNodeEditModes::TwoBoneIK;
+}
+
 #undef LOCTEXT_NAMESPACE
diff --git a/Engine/Source/Editor/AnimGraph/Private/AnimationGraphSchema.cpp b/Engine/Source/Editor/AnimGraph/Private/AnimationGraphSchema.cpp
--- a/Engine/Source/Editor/AnimGraph/Private/AnimationGraphSchema.cpp
+++ b/Engine/Source/Editor/AnimGraph/Private/AnimationGraphSchema.cpp
@@ -141,6 +141,43 @@ bool UAnimationGraphSchema::IsComponentSpacePosePin(const FEdGraphPinType& PinTy
 	return (PinType.PinCategory == Schema->PC_Struct) && (PinType.PinSubCategoryObject == ComponentSpacePoseLinkStruct);
 }
 
+bool UAnimationGraphSchema::TryCreateConnection(UEdGraphPin* A, UEdGraphPin* B) const
+{
+	UEdGraphPin* OutputPin = nullptr;
+	UEdGraphPin* InputPin = nullptr;
+
+	if(A->Direction == EEdGraphPinDirection::EGPD_Output)
+	{
+		OutputPin = A;
+		InputPin = B;
+	}
+	else
+	{
+		OutputPin = B;
+		InputPin = A;
+	}
+	check(OutputPin && InputPin);
+
+	UEdGraphNode* OutputNode = OutputPin->GetOwningNode();
+
+	if(UK2Node_Knot* RerouteNode = Cast<UK2Node_Knot>(OutputNode))
+	{
+		// Double check this is our "exec"-like line
+		bool bOutputIsPose = IsPosePin(OutputPin->PinType);
+		bool bInputIsPose = IsPosePin(InputPin->PinType);
+		bool bHavePosePin = bOutputIsPose || bInputIsPose;
+		bool bHaveWildPin = InputPin->PinType.PinCategory == PC_Wildcard || OutputPin->PinType.PinCategory == PC_Wildcard;
+
+		if(bOutputIsPose && bInputIsPose || (bHavePosePin && bHaveWildPin))
+		{
+			// Ok this is a valid exec-like line, we need to kill any connections already on the output pin
+			OutputPin->BreakAllPinLinks();
+		}
+	}
+
+	return Super::TryCreateConnection(A, B);
+}
+
 const FPinConnectionResponse UAnimationGraphSchema::DetermineConnectionResponseOfCompatibleTypedPins(const UEdGraphPin* PinA, const UEdGraphPin* PinB, const UEdGraphPin* InputPin, const UEdGraphPin* OutputPin) const
 {
 	// Enforce a tree hierarchy; where poses can only have one output (parent) connection
