Deprecated FRawMesh, replace by MeshDescription - Simplygon

diff --git a/Engine/Source/Developer/MeshMergeUtilities/Private/MeshMergeHelpers.cpp b/Engine/Source/Developer/MeshMergeUtilities/Private/MeshMergeHelpers.cpp
--- a/Engine/Source/Developer/MeshMergeUtilities/Private/MeshMergeHelpers.cpp
+++ b/Engine/Source/Developer/MeshMergeUtilities/Private/MeshMergeHelpers.cpp
@@ -1234,4 +1234,29 @@ void FMeshMergeHelpers::MergeImpostersToRawMesh(TArray<const UStaticMeshComponen
 			FMeshMergeHelpers::AppendRawMesh(InRawMesh, ImposterMesh);
 		}
 	}
+}
+
+FMeshDescription* FMeshMergeData::GetMeshDescription()
+{
+	if (MeshDescription != nullptr)
+	{
+		return MeshDescription;
+	}
+	if (RawMesh == nullptr)
+	{
+		return nullptr;
+	}
+
+	MeshDescription = new FMeshDescription();
+	UStaticMesh::RegisterMeshAttributes(*MeshDescription);
+	TMap<int32, FName> MaterialMap;
+	if (SourceStaticMesh)
+	{
+		for (int32 MaterialIndex = 0; MaterialIndex < SourceStaticMesh->StaticMaterials.Num(); ++MaterialIndex)
+		{
+			MaterialMap.Add(MaterialIndex, SourceStaticMesh->StaticMaterials[MaterialIndex].ImportedMaterialSlotName);
+		}
+	}
+	FMeshDescriptionOperations::ConvertFromRawMesh(*RawMesh, *MeshDescription, MaterialMap);
+	return MeshDescription;
 }
\ No newline at end of file
diff --git a/Engine/Source/Developer/MeshMergeUtilities/Public/MeshMergeData.h b/Engine/Source/Developer/MeshMergeUtilities/Public/MeshMergeData.h
--- a/Engine/Source/Developer/MeshMergeUtilities/Public/MeshMergeData.h
+++ b/Engine/Source/Developer/MeshMergeUtilities/Public/MeshMergeData.h
@@ -5,6 +5,7 @@
 #include "CoreMinimal.h"
 #include "Engine/StaticMesh.h"
 #include "RawMesh.h"
+#include "MeshDescription.h"
 #include "PhysicsEngine/AggregateGeom.h"
 
 #include "LightMap.h"
@@ -16,6 +17,7 @@ struct FMeshMergeData
 {
 	FMeshMergeData()
 		: RawMesh(nullptr)
+		, MeshDescription(nullptr)
 		, SourceStaticMesh(nullptr)
 		, bIsClippingMesh(false)
 	{}
@@ -27,10 +29,20 @@ struct FMeshMergeData
 			delete RawMesh;
 			RawMesh = nullptr;
 		}
+		if (MeshDescription != nullptr)
+		{
+			delete MeshDescription;
+			MeshDescription = nullptr;
+		}
 	}
 	
 	/** Raw mesh data from the source static mesh */
-	FRawMesh* RawMesh;	
+	FRawMesh* RawMesh;
+	/** MeshDesciption pointer that get populated when we query it */
+	FMeshDescription* MeshDescription;
+	//MeshDescription getter
+	MESHMERGEUTILITIES_API FMeshDescription* GetMeshDescription();
+
 	/** Contains the original texture bounds, if the material requires baking down per-vertex data */
 	TArray<FBox2D> TexCoordBounds;
 	/** Will contain non-overlapping texture coordinates, if the material requires baking down per-vertex data */	
diff --git a/Engine/Source/Developer/MeshReductionInterface/Public/IMeshReductionInterfaces.h b/Engine/Source/Developer/MeshReductionInterface/Public/IMeshReductionInterfaces.h
--- a/Engine/Source/Developer/MeshReductionInterface/Public/IMeshReductionInterfaces.h
+++ b/Engine/Source/Developer/MeshReductionInterface/Public/IMeshReductionInterfaces.h
@@ -20,6 +20,7 @@ public:
 	virtual ~IMeshReduction() {}
 
 	/**
+	* DEPRECATED Please use ReduceMeshDescription instead.
 	* Reduces the raw mesh using the provided reduction settings.
 	* @param OutReducedMesh - Upon return contains the reduced mesh.
 	* @param OutMaxDeviation - Upon return contains the maximum distance by which the reduced mesh deviates from the original.
diff --git a/Engine/Source/Developer/SimplygonMeshReduction/Private/SimplygonMeshReduction.cpp b/Engine/Source/Developer/SimplygonMeshReduction/Private/SimplygonMeshReduction.cpp
--- a/Engine/Source/Developer/SimplygonMeshReduction/Private/SimplygonMeshReduction.cpp
+++ b/Engine/Source/Developer/SimplygonMeshReduction/Private/SimplygonMeshReduction.cpp
@@ -929,7 +929,7 @@ public:
 
 		virtual uint32 Run()
 		{
-			FRawMesh OutProxyMesh;
+			FMeshDescription OutProxyMesh;
 			FFlattenMaterial OutMaterial;
 
 			if (!Data.Num())
@@ -1012,13 +1012,24 @@ public:
 			Reduction->CreateRawMeshFromGeometry(OutProxyMesh, ProxyMesh->GetGeometry(), WINDING_Keep);
 
 			// Default smoothing
-			OutProxyMesh.FaceSmoothingMasks.SetNum(OutProxyMesh.FaceMaterialIndices.Num());
-			for (uint32& SmoothingMask : OutProxyMesh.FaceSmoothingMasks)
+			TEdgeAttributesRef<bool> EdgeHardnesses = OutProxyMesh.EdgeAttributes().GetAttributesRef<bool>(MeshAttribute::Edge::IsHard);
+
+			for (const FEdgeID& EdgeID : OutProxyMesh.Edges().GetElementIDs())
 			{
-				SmoothingMask = 1;
+				EdgeHardnesses[EdgeID] = false;
 			}
 
-			Delegate.ExecuteIfBound(OutProxyMesh, OutMaterial, JobGUID);
+			FRawMesh ConvertedRawMesh;
+			TMap<FName, int32> MaterialMap;
+			TPolygonGroupAttributesConstRef<FName> PolygonGroupNames = OutProxyMesh.PolygonGroupAttributes().GetAttributesRef<FName>(MeshAttribute::PolygonGroup::ImportedMaterialSlotName);
+			int32 MaterialIndex = 0;
+			for (const FPolygonGroupID PolygonGroupID : OutProxyMesh.PolygonGroups().GetElementIDs())
+			{
+				MaterialMap.Add(PolygonGroupNames[PolygonGroupID], MaterialIndex++);
+			}
+			FMeshDescriptionOperations::ConvertToRawMesh(OutProxyMesh, ConvertedRawMesh, MaterialMap);
+
+			Delegate.ExecuteIfBound(ConvertedRawMesh, OutMaterial, JobGUID);
 
 			return 0;
 		}
@@ -1060,7 +1071,8 @@ public:
 		FScopedSlowTask SlowTask(100.f, (LOCTEXT("SimplygonProxyLOD_ProxyLOD", "Generating Proxy Mesh using Simplygon")));
 		SlowTask.MakeDialog();
 
-		FRawMesh OutProxyMesh;
+		FMeshDescription OutProxyMesh;
+		UStaticMesh::RegisterMeshAttributes(OutProxyMesh);
 		FFlattenMaterial OutMaterial;
 
 		if (!InData.Num())
@@ -1195,18 +1207,28 @@ public:
 		CreateRawMeshFromGeometry(OutProxyMesh, ProxyMesh->GetGeometry(), WINDING_Keep);
 		
 		// Default smoothing
-		OutProxyMesh.FaceSmoothingMasks.SetNum(OutProxyMesh.FaceMaterialIndices.Num());
-		for (uint32& SmoothingMask : OutProxyMesh.FaceSmoothingMasks)
+		TEdgeAttributesRef<bool> EdgeHardnesses = OutProxyMesh.EdgeAttributes().GetAttributesRef<bool>(MeshAttribute::Edge::IsHard);
+
+		for (const FEdgeID& EdgeID : OutProxyMesh.Edges().GetElementIDs())
 		{
-			SmoothingMask = 1;
+			EdgeHardnesses[EdgeID] = false;
 		}
 
-		if (!OutProxyMesh.IsValid())
+		if (OutProxyMesh.Polygons().Num() == 0)
 		{
 			FailedDelegate.ExecuteIfBound(InJobGUID, TEXT("Simplygon failed to generate a valid proxy mesh"));
 		}
+		FRawMesh ConvertedRawMesh;
+		TMap<FName, int32> MaterialMap;
+		TPolygonGroupAttributesConstRef<FName> PolygonGroupNames = OutProxyMesh.PolygonGroupAttributes().GetAttributesRef<FName>(MeshAttribute::PolygonGroup::ImportedMaterialSlotName);
+		int32 MaterialIndex = 0;
+		for (const FPolygonGroupID PolygonGroupID : OutProxyMesh.PolygonGroups().GetElementIDs())
+		{
+			MaterialMap.Add(PolygonGroupNames[PolygonGroupID], MaterialIndex++);
+		}
+		FMeshDescriptionOperations::ConvertToRawMesh(OutProxyMesh, ConvertedRawMesh, MaterialMap);
 				
-		CompleteDelegate.ExecuteIfBound(OutProxyMesh, OutMaterial, InJobGUID);
+		CompleteDelegate.ExecuteIfBound(ConvertedRawMesh, OutMaterial, InJobGUID);
 	}
 
 		EventHandler.Task = nullptr;
@@ -1739,7 +1761,158 @@ private:
 				DstTriangleIndex++;
 			}
 		}
+		
+		return GeometryData;
+	}
+
+	// This is a copy of CreateGeometryFromRawMesh with additional features for material LOD.
+	SimplygonSDK::spGeometryData CreateGeometryFromRawMesh(const FMeshDescription& RawMesh, const TArray<FBox2D>& TextureBounds, const TArray<FVector2D>& InTexCoords)
+	{
+		TVertexAttributesConstRef<FVector> VertexPositions = RawMesh.VertexAttributes().GetAttributesRef<FVector>(MeshAttribute::Vertex::Position);
+		TEdgeAttributesConstRef<bool> EdgeHardnesses = RawMesh.EdgeAttributes().GetAttributesRef<bool>(MeshAttribute::Edge::IsHard);
+		TEdgeAttributesConstRef<float> EdgeCreaseSharpnesses = RawMesh.EdgeAttributes().GetAttributesRef<float>(MeshAttribute::Edge::CreaseSharpness);
+		TPolygonGroupAttributesConstRef<FName> PolygonGroupImportedMaterialSlotNames = RawMesh.PolygonGroupAttributes().GetAttributesRef<FName>(MeshAttribute::PolygonGroup::ImportedMaterialSlotName);
+		TVertexInstanceAttributesConstRef<FVector> VertexInstanceNormals = RawMesh.VertexInstanceAttributes().GetAttributesRef<FVector>(MeshAttribute::VertexInstance::Normal);
+		TVertexInstanceAttributesConstRef<FVector> VertexInstanceTangents = RawMesh.VertexInstanceAttributes().GetAttributesRef<FVector>(MeshAttribute::VertexInstance::Tangent);
+		TVertexInstanceAttributesConstRef<float> VertexInstanceBinormalSigns = RawMesh.VertexInstanceAttributes().GetAttributesRef<float>(MeshAttribute::VertexInstance::BinormalSign);
+		TVertexInstanceAttributesConstRef<FVector4> VertexInstanceColors = RawMesh.VertexInstanceAttributes().GetAttributesRef<FVector4>(MeshAttribute::VertexInstance::Color);
+		TVertexInstanceAttributesConstRef<FVector2D> VertexInstanceUVs = RawMesh.VertexInstanceAttributes().GetAttributesRef<FVector2D>(MeshAttribute::VertexInstance::TextureCoordinate);
+
+		int32 NumVertices = RawMesh.Vertices().Num();
+		int32 NumWedges = RawMesh.VertexInstances().Num();
+
+		if (NumWedges == 0)
+		{
+			return NULL;
+		}
+
+		int32 NumTris = 0;
+		for (const FPolygonID& PolygonID : RawMesh.Polygons().GetElementIDs())
+		{
+			NumTris += RawMesh.GetPolygon(PolygonID).Triangles.Num();
+		}
+
+		TArray<uint32> FaceSmoothingMasks;
+		FMeshDescriptionOperations::ConvertHardEdgesToSmoothGroup(RawMesh, FaceSmoothingMasks);
+
+		SimplygonSDK::spGeometryData GeometryData = SDK->CreateGeometryData();
+		GeometryData->SetVertexCount(NumVertices);
+		GeometryData->SetTriangleCount(NumTris);
+
+		TMap<FVertexID, int32> VertexIDToDstVertexIndex;
+		VertexIDToDstVertexIndex.Reserve(NumVertices);
+		int32 VertexCount = 0;
+		SimplygonSDK::spRealArray Positions = GeometryData->GetCoords();
+		for (const FVertexID& VertexID : RawMesh.Vertices().GetElementIDs())
+		{
+			FVector TempPos = VertexPositions[VertexID];
+			TempPos = GetConversionMatrix().TransformPosition(TempPos);
+			Positions->SetTuple(VertexID.GetValue(), (float*)&TempPos);
+			VertexIDToDstVertexIndex.Add(VertexID, VertexCount);
+			VertexCount++;
+		}
+
+		//Prepare the tex coord
+		TArray<SimplygonSDK::spRealArray> TexCoordsArray;
+		for (int32 TexCoordIndex = 0; TexCoordIndex < VertexInstanceUVs.GetNumIndices(); ++TexCoordIndex)
+		{
+			GeometryData->AddTexCoords(TexCoordIndex);
+			SimplygonSDK::spRealArray TexCoords = GeometryData->GetTexCoords(TexCoordIndex);
+			check(TexCoords->GetTupleSize() == 2);
+			TexCoordsArray.Add(TexCoords);
+		}
+
+		//Prepare the vertex color
+		GeometryData->AddColors(0);
+		SimplygonSDK::spRealArray LinearColors = GeometryData->GetColors(0);
+		check(LinearColors);
+		check(LinearColors->GetTupleSize() == 4);
+
+		//Prepare the tangent space
+		GeometryData->AddTangents(0);
+		GeometryData->AddNormals();
+		SimplygonSDK::spRealArray Tangents = GeometryData->GetTangents(0);
+		SimplygonSDK::spRealArray Bitangents = GeometryData->GetBitangents(0);
+		SimplygonSDK::spRealArray Normals = GeometryData->GetNormals();
+
+		SimplygonSDK::spRidArray Indices = GeometryData->GetVertexIds();
+
+		GeometryData->AddMaterialIds();
+		SimplygonSDK::spRidArray MaterialIndices = GeometryData->GetMaterialIds();
+
+		uint32 DstVertexIndex = 0;
+		uint32 DstTriangleIndex = 0;
+		for (const FPolygonID& PolygonID : RawMesh.Polygons().GetElementIDs())
+		{
+			for (const FMeshTriangle& Triangle : RawMesh.GetPolygon(PolygonID).Triangles)
+			{
+				//Materials
+				const FPolygonGroupID& PolygonGroupID = RawMesh.GetPolygonPolygonGroup(PolygonID);
+				MaterialIndices->SetItem(DstTriangleIndex, PolygonGroupID.GetValue());
+				int32 MaterialIndex = PolygonGroupID.GetValue();
+
+				for (int32 Corner = 0; Corner < 3; ++Corner)
+				{
+					const FVertexInstanceID VertexInstanceID = Triangle.GetVertexInstanceID(Corner);
+					//Add all the per indice data
+					Indices->SetItem(DstVertexIndex, VertexIDToDstVertexIndex[RawMesh.GetVertexInstanceVertex(VertexInstanceID)]);
+
+					//UVs
+					for (int32 TexCoordIndex = 0; TexCoordIndex < VertexInstanceUVs.GetNumIndices(); ++TexCoordIndex)
+					{
+						// Compute texture bounds for current material.
+						float MinU = 0, ScaleU = 1;
+						float MinV = 0, ScaleV = 1;
+
+						if (TextureBounds.IsValidIndex(MaterialIndex) && TexCoordIndex == 0 && InTexCoords.Num() == 0)
+						{
+							const FBox2D& Bounds = TextureBounds[MaterialIndex];
+							if (Bounds.GetArea() > 0)
+							{
+								MinU = Bounds.Min.X;
+								MinV = Bounds.Min.Y;
+								ScaleU = 1.0f / (Bounds.Max.X - Bounds.Min.X);
+								ScaleV = 1.0f / (Bounds.Max.Y - Bounds.Min.Y);
+							}
+						}
+
+						const FVector2D& TexCoord = VertexInstanceUVs.Get(VertexInstanceID, TexCoordIndex);
+						float UV[2];
+						UV[0] = (TexCoord.X - MinU) * ScaleU;
+						UV[1] = (TexCoord.Y - MinV) * ScaleV;
+						TexCoordsArray[TexCoordIndex]->SetTuple(DstVertexIndex, UV);
+					}
 
+					//Colors
+					LinearColors->SetTuple(DstVertexIndex, (float*)&VertexInstanceColors[VertexInstanceID]);
+
+					//Tangents
+					FVector TempTangent = VertexInstanceTangents[VertexInstanceID];
+					TempTangent = GetConversionMatrix().TransformVector(TempTangent);
+					Tangents->SetTuple(DstVertexIndex, (float*)&TempTangent);
+
+					FVector TempBitangent = FVector::CrossProduct(VertexInstanceNormals[VertexInstanceID], VertexInstanceTangents[VertexInstanceID]).GetSafeNormal() * VertexInstanceBinormalSigns[VertexInstanceID];
+					TempBitangent = GetConversionMatrix().TransformVector(TempBitangent);
+					Bitangents->SetTuple(DstVertexIndex, (float*)&TempBitangent);
+
+					FVector TempNormal = VertexInstanceNormals[VertexInstanceID];
+					TempNormal = GetConversionMatrix().TransformVector(TempNormal);
+					Normals->SetTuple(DstVertexIndex, (float*)&TempNormal);
+
+					//Increment the indice index
+					DstVertexIndex++;
+				}
+
+				// Per-triangle data.
+
+				//Smooth group
+				GeometryData->AddGroupIds();
+				SimplygonSDK::spRidArray GroupIds = GeometryData->GetGroupIds();
+				GroupIds->SetItem(DstTriangleIndex, FaceSmoothingMasks[DstTriangleIndex]);
+
+				DstTriangleIndex++;
+			}
+		}
 		return GeometryData;
 	}
 
@@ -1806,13 +1979,13 @@ private:
 			check(TexCoords->GetTupleSize() == 2);
 			TexCoordsArray.Add(TexCoords);
 		}
-		VertexInstanceUVs.SetNumIndices(TexCoordNum);
+		VertexInstanceUVs.SetNumIndices(FMath::Max(TexCoordNum, 1));
 
 		//Prepare the polygongroup
 		TMap<int32, FPolygonGroupID> GeoToRawMaterial;
 		for (int32 TriIndex = 0; TriIndex < NumTris; ++TriIndex)
 		{
-			int32 MaterialIndex = MaterialIndices->GetItem(TriIndex);
+			int32 MaterialIndex = MaterialIndices != nullptr ? MaterialIndices->GetItem(TriIndex) : 0;
 			if (!GeoToRawMaterial.Contains(MaterialIndex))
 			{
 				const FPolygonGroupID PolygonGroupID(MaterialIndex);
@@ -1844,22 +2017,25 @@ private:
 				}
 
 				//Vertex Color
-				LinearColors->GetTuple(SrcIndex, sgTuple);
-				SimplygonSDK::real* sgVertexColor = sgTuple->GetData();
-				VertexInstanceColors[VertexInstanceID] = FVector4(sgVertexColor[0], sgVertexColor[1], sgVertexColor[2], sgVertexColor[3]);
+				if (LinearColors)
+				{
+					LinearColors->GetTuple(SrcIndex, sgTuple);
+					SimplygonSDK::real* sgVertexColor = sgTuple->GetData();
+					VertexInstanceColors[VertexInstanceID] = FVector4(sgVertexColor[0], sgVertexColor[1], sgVertexColor[2], sgVertexColor[3]);
+				}
 
 				//Vertex Tangents
 				FVector Tangent(0.0f);
 				FVector BiTangent(0.0f);
 				FVector Normal(0.0f);
 
-				Normals->GetTuple(SrcIndex, sgTuple);
-				SimplygonSDK::real* sgNormal = sgTuple->GetData();
-				Normal = GetConversionMatrix().TransformVector(FVector(sgNormal[0], sgNormal[1], sgNormal[2]));
-				if (Normal.IsNearlyZero())
+				if (Normals)
 				{
-					bHasZeroNormal = true;
+					Normals->GetTuple(SrcIndex, sgTuple);
+					SimplygonSDK::real* sgNormal = sgTuple->GetData();
+					Normal = GetConversionMatrix().TransformVector(FVector(sgNormal[0], sgNormal[1], sgNormal[2]));
 				}
+				bHasZeroNormal |= Normal.IsNearlyZero();
 
 				if (Tangents && Bitangents)
 				{
@@ -1871,12 +2047,8 @@ private:
 					SimplygonSDK::real* sgBiTangents = sgTuple->GetData();
 					BiTangent = GetConversionMatrix().TransformVector(FVector(sgBiTangents[0], sgBiTangents[1], sgBiTangents[2]));
 				}
-				else
-				{
-					bHasZeroTangent = true;
-					Tangent = FVector(0.0f);
-					BiTangent = FVector(0.0f);
-				}
+				bHasZeroTangent |= Tangent.IsNearlyZero() || BiTangent.IsNearlyZero();
+
 				VertexInstanceTangents[VertexInstanceID] = Tangent;
 				VertexInstanceBinormalSigns[VertexInstanceID] = GetBasisDeterminantSign(Tangent.GetSafeNormal(), BiTangent.GetSafeNormal(), Normal.GetSafeNormal());
 				VertexInstanceNormals[VertexInstanceID] = Normal;
@@ -1907,7 +2079,8 @@ private:
 				ContourPoint.VertexInstanceID = VertexInstanceIDs[CornerIndices[0]];
 			}
 			// Insert a polygon into the mesh
-			const FPolygonID NewPolygonID = RawMesh.CreatePolygon(GeoToRawMaterial[MaterialIndices->GetItem(TriIndex)], Contours);
+			int32 MaterialIndex = MaterialIndices != nullptr ? MaterialIndices->GetItem(TriIndex) : 0;
+			const FPolygonID NewPolygonID = RawMesh.CreatePolygon(GeoToRawMaterial[MaterialIndex], Contours);
 			//Triangulate the polygon
 			FMeshPolygon& Polygon = RawMesh.GetPolygon(NewPolygonID);
 			RawMesh.ComputePolygonTriangulation(NewPolygonID, Polygon.Triangles);
diff --git a/Engine/Source/Developer/SimplygonSwarm/Private/SimplygonSwarm.cpp b/Engine/Source/Developer/SimplygonSwarm/Private/SimplygonSwarm.cpp
--- a/Engine/Source/Developer/SimplygonSwarm/Private/SimplygonSwarm.cpp
+++ b/Engine/Source/Developer/SimplygonSwarm/Private/SimplygonSwarm.cpp
@@ -24,6 +24,11 @@ THIRD_PARTY_INCLUDES_END
 #include "MeshMergeData.h"
 #include "Features/IModularFeatures.h"
 
+#include "MeshDescription.h"
+#include "MeshAttributes.h"
+#include "MeshAttributeArray.h"
+#include "MeshDescriptionOperations.h"
+
 // Standard Simplygon channels have some issues with extracting color data back from simplification, 
 // so we use this workaround with user channels
 static const char* USER_MATERIAL_CHANNEL_METALLIC = "UserMetallic";
@@ -67,7 +72,7 @@ static const TCHAR* SG_UE_INTEGRATION_REV = TEXT("#SG_UE_INTEGRATION_REV");
 
 static const TCHAR SHADING_NETWORK_TEMPLATE[] = TEXT("<SimplygonShadingNetwork version=\"1.0\">\n\t<ShadingTextureNode ref=\"node_0\" name=\"ShadingTextureNode\">\n\t\t<DefaultColor0>\n\t\t\t<DefaultValue>1 1 1 1</DefaultValue>\n\t\t</DefaultColor0>\n\t\t<TextureName>%s</TextureName>\n\t\t<TextureLevelName>%s</TextureLevelName>\n\t\t<UseSRGB>%d</UseSRGB>\n\t\t<TileU>1.000000</TileU>\n\t\t<TileV>1.000000</TileV>\n\t</ShadingTextureNode>\n</SimplygonShadingNetwork>");
 
-ssf::pssfMeshData CreateSSFMeshDataFromRawMesh(const FRawMesh& InRawMesh, TArray<FBox2D> InTextureBounds, TArray<FVector2D> InTexCoords);
+ssf::pssfMeshData CreateSSFMeshDataFromRawMesh(const FMeshDescription& InRawMesh, TArray<FBox2D> InTextureBounds, TArray<FVector2D> InTexCoords);
 
 class FSimplygonSwarmModule : public IMeshReductionModule
 {
@@ -140,7 +145,7 @@ public:
 		FScopedSlowTask SlowTask(3.f, (LOCTEXT("SimplygonSwarm_ProxyLOD", "Generating Proxy Mesh using Simplygon Swarm")));
 		SlowTask.MakeDialog();
 		
-		FRawMesh OutProxyMesh;
+		FMeshDescription OutProxyMesh;
 		FFlattenMaterial OutMaterial;
 
 		//setup path variables
@@ -199,11 +204,11 @@ public:
 
 		ssf::pssfScene SsfScene;
 		
-		TArray<FRawMesh*> InputMeshes;
+		TArray<FMeshDescription*> InputMeshes;
 
 		for (auto Data : InData)
 		{
-			InputMeshes.Push(Data.RawMesh);
+			InputMeshes.Push(Data.GetMeshDescription());
 		}
 
 		bool bDiscardEmissive = true;
@@ -311,7 +316,7 @@ public:
 	void ImportFile(const FSimplygonSwarmTask& InSwarmTask)
 	{
 
-		FRawMesh OutProxyMesh;
+		FMeshDescription OutProxyMesh;
 		FFlattenMaterial OutMaterial;
 		bool bDebuggingEnabled = GetDefault<UEditorPerProjectUserSettings>()->bEnableSwarmDebugging;
 		FString OutputFolderPath = FString::Printf(TEXT("%s/Output"), *InSwarmTask.TaskData.JobDirectory);
@@ -343,10 +348,10 @@ public:
 				OutMaterial.SetPropertySize(EFlattenMaterialProperties::Emissive, FIntPoint(0,0));
 			}
 
-			if (!OutProxyMesh.IsValid())
+			if (!OutProxyMesh.VertexInstances().Num())
 			{
 				UE_LOG(LogSimplygonSwarm, Log, TEXT("RawMesh is invalid."));
-				FailedDelegate.ExecuteIfBound(InSwarmTask.TaskData.ProcessorJobID, TEXT("Invalid FRawMesh data"));
+				FailedDelegate.ExecuteIfBound(InSwarmTask.TaskData.ProcessorJobID, TEXT("Invalid FMeshDescription data"));
 			}
 
 			 
@@ -368,7 +373,17 @@ public:
 			//if is bound then execute
 			if (CompleteDelegate.IsBound())
 			{
-				CompleteDelegate.Execute(OutProxyMesh, OutMaterial, InSwarmTask.TaskData.ProcessorJobID);
+				FRawMesh ConvertedRawMesh;
+				TMap<FName, int32> MaterialMap;
+				TPolygonGroupAttributesConstRef<FName> PolygonGroupNames = OutProxyMesh.PolygonGroupAttributes().GetAttributesRef<FName>(MeshAttribute::PolygonGroup::ImportedMaterialSlotName);
+				int32 MaterialIndex = 0;
+				for (const FPolygonGroupID PolygonGroupID : OutProxyMesh.PolygonGroups().GetElementIDs())
+				{
+					MaterialMap.Add(PolygonGroupNames[PolygonGroupID], MaterialIndex++);
+				}
+				FMeshDescriptionOperations::ConvertToRawMesh(OutProxyMesh, ConvertedRawMesh, MaterialMap);
+
+				CompleteDelegate.Execute(ConvertedRawMesh, OutMaterial, InSwarmTask.TaskData.ProcessorJobID);
 			}
 			else
 			{
@@ -379,7 +394,7 @@ public:
 		else
 		{
 			UE_LOG(LogSimplygonSwarm, Log, TEXT("Failed to unzip downloaded content %s"), *ZipFileFullPath);
-			FailedDelegate.ExecuteIfBound(InSwarmTask.TaskData.ProcessorJobID, TEXT("Invalid FRawMesh data"));
+			FailedDelegate.ExecuteIfBound(InSwarmTask.TaskData.ProcessorJobID, TEXT("Invalid FMeshDescription data"));
 		}
 	}
 
@@ -578,7 +593,7 @@ private:
 			Count++;
 
 			//setup mesh data
-			ssf::pssfMeshData SsfMeshData = CreateSSFMeshDataFromRawMesh(*MergeData.RawMesh, MergeData.TexCoordBounds, MergeData.NewUVs);
+			ssf::pssfMeshData SsfMeshData = CreateSSFMeshDataFromRawMesh(*MergeData.GetMeshDescription(), MergeData.TexCoordBounds, MergeData.NewUVs);
 			SsfMesh->MeshDataList.push_back(SsfMeshData);
 
 			//setup mesh material information
@@ -587,7 +602,7 @@ private:
 			UniqueMaterialIds.Reserve(InputMaterials.Num());
 
 			//get unqiue material ids
-			GetUniqueMaterialIndices(MergeData.RawMesh->FaceMaterialIndices, UniqueMaterialIds);
+			GetUniqueMaterialIndices(*MergeData.GetMeshDescription(), UniqueMaterialIds);
 
 			SsfMesh->MaterialIds->Items.reserve(UniqueMaterialIds.Num());
 
@@ -643,8 +658,18 @@ private:
 	* @param OutMaterial		Converted SsfMaterial to Flattened Material
 	* @param BaseTexturesPath	Base Path for textures
 	*/
-	void ConvertFromSsfSceneToRawMesh(ssf::pssfScene SsfScene, FRawMesh& OutProxyMesh, FFlattenMaterial& OutMaterial, const FString BaseTexturesPath)
+	void ConvertFromSsfSceneToRawMesh(ssf::pssfScene SsfScene, FMeshDescription& OutProxyMesh, FFlattenMaterial& OutMaterial, const FString BaseTexturesPath)
 	{
+		TVertexAttributesRef<FVector> VertexPositions = OutProxyMesh.VertexAttributes().GetAttributesRef<FVector>(MeshAttribute::Vertex::Position);
+		TEdgeAttributesRef<bool> EdgeHardnesses = OutProxyMesh.EdgeAttributes().GetAttributesRef<bool>(MeshAttribute::Edge::IsHard);
+		TEdgeAttributesRef<float> EdgeCreaseSharpnesses = OutProxyMesh.EdgeAttributes().GetAttributesRef<float>(MeshAttribute::Edge::CreaseSharpness);
+		TPolygonGroupAttributesRef<FName> PolygonGroupImportedMaterialSlotNames = OutProxyMesh.PolygonGroupAttributes().GetAttributesRef<FName>(MeshAttribute::PolygonGroup::ImportedMaterialSlotName);
+		TVertexInstanceAttributesRef<FVector> VertexInstanceNormals = OutProxyMesh.VertexInstanceAttributes().GetAttributesRef<FVector>(MeshAttribute::VertexInstance::Normal);
+		TVertexInstanceAttributesRef<FVector> VertexInstanceTangents = OutProxyMesh.VertexInstanceAttributes().GetAttributesRef<FVector>(MeshAttribute::VertexInstance::Tangent);
+		TVertexInstanceAttributesRef<float> VertexInstanceBinormalSigns = OutProxyMesh.VertexInstanceAttributes().GetAttributesRef<float>(MeshAttribute::VertexInstance::BinormalSign);
+		TVertexInstanceAttributesRef<FVector4> VertexInstanceColors = OutProxyMesh.VertexInstanceAttributes().GetAttributesRef<FVector4>(MeshAttribute::VertexInstance::Color);
+		TVertexInstanceAttributesRef<FVector2D> VertexInstanceUVs = OutProxyMesh.VertexInstanceAttributes().GetAttributesRef<FVector2D>(MeshAttribute::VertexInstance::TextureCoordinate);
+
 		bool bReverseWinding = true;
 		 
 		for (ssf::pssfMesh Mesh : SsfScene->MeshTable->MeshList)
@@ -653,142 +678,139 @@ private:
 			for (ssf::pssfMeshData MeshData : Mesh->MeshDataList)
 			{
 				int32 TotalVertices = MeshData->GetVerticesCount();
-				int32 ToatalCorners = MeshData->GetCornersCount();
+				int32 TotalCorners = MeshData->GetCornersCount();
 				int32 TotalTriangles = MeshData->GetTrianglesCount();
 
-				OutProxyMesh.VertexPositions.SetNumUninitialized(TotalVertices);
-				int VertexIndex = 0;
+				//Assuming only one mesh
+				OutProxyMesh.Empty();
+				OutProxyMesh.ReserveNewVertices(TotalVertices);
+				OutProxyMesh.ReserveNewPolygons(TotalTriangles);
+				OutProxyMesh.ReserveNewVertexInstances(TotalCorners);
+				OutProxyMesh.ReserveNewEdges(TotalCorners);
+
+				TMap<int32, FVertexID> SsfToMeshDescriptionVertexID;
+				SsfToMeshDescriptionVertexID.Reserve(TotalVertices);
+				int32 VertexIndex = 0;
 				for (ssf::ssfVector3 VertexCoord : MeshData->Coordinates.Get().Items)
 				{
-					OutProxyMesh.VertexPositions[VertexIndex] = GetConversionMatrixYUP().InverseTransformPosition(FVector(VertexCoord.V[0], VertexCoord.V[1], VertexCoord.V[2]));
+					const FVertexID VertexID = OutProxyMesh.CreateVertex();
+					VertexPositions[VertexID] = GetConversionMatrixYUP().InverseTransformPosition(FVector(VertexCoord.V[0], VertexCoord.V[1], VertexCoord.V[2]));
+					SsfToMeshDescriptionVertexID.Add(VertexIndex, VertexID);
 					VertexIndex++;
 				}
-				 
-				 
-				OutProxyMesh.WedgeIndices.SetNumUninitialized(ToatalCorners);
-				for (int32 TriIndex = 0; TriIndex < TotalTriangles; ++TriIndex)
-				{
-					for (int32 CornerIndex = 0; CornerIndex < 3; ++CornerIndex)
-					{
-						int32 DestCornerIndex = bReverseWinding ? 2 - CornerIndex : CornerIndex;
-						OutProxyMesh.WedgeIndices[TriIndex * 3 + DestCornerIndex] = MeshData->TriangleIndices.Get().Items[TriIndex].V[CornerIndex];
-					}					 
-				}
-
-				//Note : Since we are doing mesh aggregation need to make sure to extract MaterialLOD TexCoord and Lightmap TexCoords
-
-				//Copy baked material UV's only discard the reset
+				
+				//Prepare the tex coord
 				int32 TexCoordIndex = 0;
 				ssf::ssfNamedList<ssf::ssfVector2> BakedMaterialUVs = FSimplygonSSFHelper::GetBakedMaterialUVs(MeshData->TextureCoordinatesList);
-					OutProxyMesh.WedgeTexCoords[TexCoordIndex].SetNumUninitialized(ToatalCorners);
-					for (int32 TriIndex = 0; TriIndex < TotalTriangles; ++TriIndex)
-					{						 
-						for (int32 CornerIndex = 0; CornerIndex < 3; ++CornerIndex)
-						{
-							int32 DestCornerIndex =  bReverseWinding ? 2 - CornerIndex : CornerIndex;
-						OutProxyMesh.WedgeTexCoords[TexCoordIndex][TriIndex * 3 + DestCornerIndex].X = BakedMaterialUVs.Items[TriIndex * 3 + CornerIndex].V[0];
-						OutProxyMesh.WedgeTexCoords[TexCoordIndex][TriIndex * 3 + DestCornerIndex].Y = BakedMaterialUVs.Items[TriIndex * 3 + CornerIndex].V[1];
-
-						}
-				}
-
+				VertexInstanceUVs.SetNumIndices(1);
 
-				//SSF Can store multiple color channels. However UE only supports one color channel
-				int32 ColorChannelIndex = 0;
-				for (ssf::ssfNamedList<ssf::ssfVector4> TexCoorChannel : MeshData->ColorsList)
-				{
-					OutProxyMesh.WedgeColors.SetNumUninitialized(ToatalCorners);
-					for (int TriIndex = 0; TriIndex < TotalTriangles; ++TriIndex)
-					{
-						for (int32 CornerIndex = 0; CornerIndex < 2; ++CornerIndex)
-						{
-							int32 DestCornerIndex = bReverseWinding ? 2 - CornerIndex : CornerIndex;
-							OutProxyMesh.WedgeColors[TriIndex * 3 +DestCornerIndex].R = TexCoorChannel.Items[TriIndex * 3 +CornerIndex].V[0];
-							OutProxyMesh.WedgeColors[TriIndex * 3 +DestCornerIndex].G = TexCoorChannel.Items[TriIndex * 3 +CornerIndex].V[1];
-							OutProxyMesh.WedgeColors[TriIndex * 3 +DestCornerIndex].B = TexCoorChannel.Items[TriIndex * 3 +CornerIndex].V[2];
-							OutProxyMesh.WedgeColors[TriIndex * 3 +DestCornerIndex].A = TexCoorChannel.Items[TriIndex * 3 +CornerIndex].V[3];
-						}
-					}
-					 
-				}
-				
-				 
+				//Is buffer has some data?
 				bool Normals = !MeshData->Normals.IsEmpty() && MeshData->Normals.Get().Items.size() > 0;
 				bool Tangents = !MeshData->Tangents.IsEmpty() && MeshData->Tangents.Get().Items.size() > 0;
 				bool Bitangents = !MeshData->Bitangents.IsEmpty() && MeshData->Bitangents.Get().Items.size() > 0;
 				bool MaterialIndices = !MeshData->MaterialIndices.IsEmpty() && MeshData->MaterialIndices.Get().Items.size() > 0;
 				bool GroupIds = !MeshData->SmoothingGroup.IsEmpty() && MeshData->SmoothingGroup.Get().Items.size() > 0;
 
-				if (Normals)
+				//Setup PolygonGroup
+				//Prepare the polygongroup
+				TMap<int32, FPolygonGroupID> SsfToRawMaterial;
+				for (int32 TriIndex = 0; TriIndex < TotalTriangles; ++TriIndex)
+				{
+					int32 MaterialIndex = MaterialIndices ? (int32)MeshData->MaterialIndices.Get().Items[TriIndex].Value : 0;
+					if (!SsfToRawMaterial.Contains(MaterialIndex))
+					{
+						const FPolygonGroupID PolygonGroupID(MaterialIndex);
+						OutProxyMesh.CreatePolygonGroupWithID(PolygonGroupID);
+						PolygonGroupImportedMaterialSlotNames[PolygonGroupID] = FName(*FString(TEXT("SimplygonSwarm_") + FString::FromInt(PolygonGroupID.GetValue())));
+						SsfToRawMaterial.Add(MaterialIndex, PolygonGroupID);
+					}
+				}
+
+				for (int32 TriIndex = 0; TriIndex < TotalTriangles; ++TriIndex)
 				{
-					if (Tangents && Bitangents)
+					FVertexID VertexIndexes[3];
+					FVertexInstanceID VertexInstanceIDs[3];
+					for (int32 CornerIndex = 0; CornerIndex < 3; ++CornerIndex)
 					{
-						OutProxyMesh.WedgeTangentX.SetNumUninitialized(ToatalCorners);
-						OutProxyMesh.WedgeTangentY.SetNumUninitialized(ToatalCorners);
+						int32 SrcCornerIndex = bReverseWinding ? 2 - CornerIndex : CornerIndex;
+						int32 SrcIndex = (3 * TriIndex) + SrcCornerIndex;
+						VertexIndexes[CornerIndex] = SsfToMeshDescriptionVertexID[MeshData->TriangleIndices.Get().Items[TriIndex].V[SrcCornerIndex]];
+						VertexInstanceIDs[CornerIndex] = OutProxyMesh.CreateVertexInstance(VertexIndexes[CornerIndex]);
+
+						//Texture Coordinates, copy baked material UV's only discard the rest
+						VertexInstanceUVs.Get(VertexInstanceIDs[CornerIndex], 0) = FVector2D(BakedMaterialUVs.Items[SrcIndex].V[0], BakedMaterialUVs.Items[SrcIndex].V[1]);
 
-						for (int32 TriIndex = 0; TriIndex < TotalTriangles; ++TriIndex)
+						//Vertex Color, SSF can store multiple color channels. However UE only supports one color channel
+						for (ssf::ssfNamedList<ssf::ssfVector4> TexColorChannel : MeshData->ColorsList)
 						{
-							for (int32 CornerIndex = 0; CornerIndex < 3; ++CornerIndex)
-							{
-								int32 DestCornerIndex = bReverseWinding ? 2 - CornerIndex : CornerIndex;
-								OutProxyMesh.WedgeTangentX[TriIndex * 3 + DestCornerIndex].X = MeshData->Tangents.Get().Items[TriIndex * 3 + CornerIndex].V[0];
-								OutProxyMesh.WedgeTangentX[TriIndex * 3 + DestCornerIndex].Y = MeshData->Tangents.Get().Items[TriIndex * 3 + CornerIndex].V[1];
-								OutProxyMesh.WedgeTangentX[TriIndex * 3 + DestCornerIndex].Z = MeshData->Tangents.Get().Items[TriIndex * 3 + CornerIndex].V[2];
-								OutProxyMesh.WedgeTangentX[TriIndex * 3 + DestCornerIndex] = GetConversionMatrixYUP().InverseTransformPosition(OutProxyMesh.WedgeTangentX[TriIndex * 3 + DestCornerIndex]);
-							}
+							VertexInstanceColors[VertexInstanceIDs[CornerIndex]] = FVector4(TexColorChannel.Items[SrcIndex].V[0], TexColorChannel.Items[SrcIndex].V[1], TexColorChannel.Items[SrcIndex].V[2], TexColorChannel.Items[SrcIndex].V[3]);
+							break; //UE support only one
 						}
 
-						for (int32 TriIndex = 0; TriIndex < TotalTriangles; ++TriIndex)
+						//Tangents
+						if (Normals)
 						{
-							for (int32 CornerIndex = 0; CornerIndex < 3; ++CornerIndex)
+							FVector NormalValue = FVector(MeshData->Normals.Get().Items[SrcIndex].V[0], MeshData->Normals.Get().Items[SrcIndex].V[1], MeshData->Normals.Get().Items[SrcIndex].V[2]);
+							NormalValue = GetConversionMatrixYUP().InverseTransformPosition(NormalValue);
+							VertexInstanceNormals[VertexInstanceIDs[CornerIndex]] = NormalValue;
+							if (Tangents && Bitangents)
 							{
-								int32 DestCornerIndex = bReverseWinding ? 2 - CornerIndex : CornerIndex;
-								OutProxyMesh.WedgeTangentY[TriIndex * 3 + DestCornerIndex].X = MeshData->Bitangents.Get().Items[TriIndex * 3 + CornerIndex].V[0];
-								OutProxyMesh.WedgeTangentY[TriIndex * 3 + DestCornerIndex].Y = MeshData->Bitangents.Get().Items[TriIndex * 3 + CornerIndex].V[1];
-								OutProxyMesh.WedgeTangentY[TriIndex * 3 + DestCornerIndex].Z = MeshData->Bitangents.Get().Items[TriIndex * 3 + CornerIndex].V[2];
-								OutProxyMesh.WedgeTangentY[TriIndex * 3 + DestCornerIndex] = GetConversionMatrixYUP().InverseTransformPosition(OutProxyMesh.WedgeTangentY[TriIndex * 3 + DestCornerIndex]);
+								FVector TangentValue = FVector(MeshData->Tangents.Get().Items[SrcIndex].V[0], MeshData->Tangents.Get().Items[SrcIndex].V[1], MeshData->Tangents.Get().Items[SrcIndex].V[2]);
+								TangentValue = GetConversionMatrixYUP().InverseTransformPosition(TangentValue);
+								VertexInstanceTangents[VertexInstanceIDs[CornerIndex]] = TangentValue;
+								
+								FVector BiTangentValue = FVector(MeshData->Bitangents.Get().Items[SrcIndex].V[0], MeshData->Bitangents.Get().Items[SrcIndex].V[1], MeshData->Bitangents.Get().Items[SrcIndex].V[2]);
+								BiTangentValue = GetConversionMatrixYUP().InverseTransformPosition(BiTangentValue);
+								VertexInstanceBinormalSigns[VertexInstanceIDs[CornerIndex]] = GetBasisDeterminantSign(TangentValue.GetSafeNormal(), BiTangentValue.GetSafeNormal(), NormalValue.GetSafeNormal());
 							}
 						}
-
-						 
-						 
 					}
-
-					OutProxyMesh.WedgeTangentZ.SetNumUninitialized(ToatalCorners);
-					for (int32 TriIndex = 0; TriIndex < TotalTriangles; ++TriIndex)
+					//Create a polygon from this triangle
+					TArray<FMeshDescription::FContourPoint> Contours;
+					for (int32 Corner = 0; Corner < 3; ++Corner)
 					{
-						for (int32 CornerIndex = 0; CornerIndex < 3; ++CornerIndex)
+						int32 ContourPointIndex = Contours.AddDefaulted();
+						FMeshDescription::FContourPoint& ContourPoint = Contours[ContourPointIndex];
+						//Find the matching edge ID
+						uint32 CornerIndices[2];
+						CornerIndices[0] = (Corner + 0) % 3;
+						CornerIndices[1] = (Corner + 1) % 3;
+
+						FVertexID EdgeVertexIDs[2];
+						EdgeVertexIDs[0] = VertexIndexes[CornerIndices[0]];
+						EdgeVertexIDs[1] = VertexIndexes[CornerIndices[1]];
+
+						FEdgeID MatchEdgeId = OutProxyMesh.GetVertexPairEdge(EdgeVertexIDs[0], EdgeVertexIDs[1]);
+						if (MatchEdgeId == FEdgeID::Invalid)
 						{
-							int32 DestCornerIndex = bReverseWinding ? 2 - CornerIndex : CornerIndex;
-							OutProxyMesh.WedgeTangentZ[TriIndex * 3 + DestCornerIndex].X = MeshData->Normals.Get().Items[TriIndex * 3 + CornerIndex].V[0];
-							OutProxyMesh.WedgeTangentZ[TriIndex * 3 + DestCornerIndex].Y = MeshData->Normals.Get().Items[TriIndex * 3 + CornerIndex].V[1];
-							OutProxyMesh.WedgeTangentZ[TriIndex * 3 + DestCornerIndex].Z = MeshData->Normals.Get().Items[TriIndex * 3 + CornerIndex].V[2];
-							OutProxyMesh.WedgeTangentZ[TriIndex * 3 + DestCornerIndex] = GetConversionMatrixYUP().InverseTransformPosition(OutProxyMesh.WedgeTangentZ[TriIndex * 3 + DestCornerIndex]);
+							MatchEdgeId = OutProxyMesh.CreateEdge(EdgeVertexIDs[0], EdgeVertexIDs[1]);
+							EdgeHardnesses[MatchEdgeId] = false;
+							EdgeCreaseSharpnesses[MatchEdgeId] = 0.0f;
 						}
+						ContourPoint.EdgeID = MatchEdgeId;
+						ContourPoint.VertexInstanceID = VertexInstanceIDs[CornerIndices[0]];
 					}
-					 
-				}
-
-				OutProxyMesh.FaceMaterialIndices.SetNumUninitialized(TotalTriangles);
-				if (MaterialIndices)
-				{
-					for (int32 TriIndex = 0; TriIndex < TotalTriangles; ++TriIndex)
-					{
-						OutProxyMesh.FaceMaterialIndices[TriIndex] = MeshData->MaterialIndices.Get().Items[TriIndex].Value;
-					}
+					// Insert a polygon into the mesh
+					const FPolygonID NewPolygonID = OutProxyMesh.CreatePolygon(SsfToRawMaterial[MeshData->MaterialIndices.Get().Items[TriIndex].Value], Contours);
+					//Triangulate the polygon
+					FMeshPolygon& Polygon = OutProxyMesh.GetPolygon(NewPolygonID);
+					OutProxyMesh.ComputePolygonTriangulation(NewPolygonID, Polygon.Triangles);
 				}
 
-				OutProxyMesh.FaceSmoothingMasks.SetNumUninitialized(TotalTriangles);
+				TArray<uint32> FaceSmoothingMasks;
+				FaceSmoothingMasks.AddZeroed(TotalTriangles);
 				if (GroupIds)
 				{
 					for (int32 TriIndex = 0; TriIndex < TotalTriangles; ++TriIndex)
 					{
-						OutProxyMesh.FaceSmoothingMasks[TriIndex] = MeshData->SmoothingGroup.Get().Items[TriIndex].Value;
+						FaceSmoothingMasks[TriIndex] = MeshData->SmoothingGroup.Get().Items[TriIndex].Value;
 					}
 				}
-
+				FMeshDescriptionOperations::ConvertSmoothGroupToHardEdges(FaceSmoothingMasks, OutProxyMesh);
 			}
 
+			
+
 			//since its a proxy will only contain one material on it
 			ssf::ssfString ProxyMaterialGuid = Mesh->MaterialIds.Get().Items[0].Value;
 			ssf::pssfMaterial ProxyMaterial = FSimplygonSSFHelper::FindMaterialById(SsfScene, ProxyMaterialGuid);
@@ -1015,11 +1037,13 @@ bool ZipContentsForUpload(FString InputDirectoryPath, FString OutputFileName)
 	* @param OriginalMaterialIds		Original Material Indicies
 	* @param ChannelUniqueMaterialIds	OutUniqueMaterialIds
 	*/
-	void GetUniqueMaterialIndices(const TArray<int32>& OriginalMaterialIds, TArray<int32>& UniqueMaterialIds)
+	void GetUniqueMaterialIndices(const FMeshDescription& MeshDescription, TArray<int32>& UniqueMaterialIds)
 	{
-		for (int32 index : OriginalMaterialIds)
+		int32 index = 0;
+		for (const FPolygonGroupID PolygonGroupID : MeshDescription.PolygonGroups().GetElementIDs())
 		{
 			UniqueMaterialIds.AddUnique(index);
+			index++;
 		}
 	}
 	
@@ -1209,20 +1233,32 @@ bool ZipContentsForUpload(FString InputDirectoryPath, FString OutputFileName)
 	}
 
 	/**
-	* Method to create a SsfMeshData from FRawMesh
+	* Method to create a SsfMeshData from FMeshDescription
 	* @param InRawMesh				Rawmesh to create SsfMeshData from
 	* @param InTextureBounds		Texture bounds
 	* @param InTexCoords			Corrected texture coordinates generated after material flattening.
 	*/
-	ssf::pssfMeshData CreateSSFMeshDataFromRawMesh(const FRawMesh& InRawMesh, TArray<FBox2D> InTextureBounds, TArray<FVector2D> InTexCoords)
+	ssf::pssfMeshData CreateSSFMeshDataFromRawMesh(const FMeshDescription& SrcRawMesh, TArray<FBox2D> InTextureBounds, TArray<FVector2D> InTexCoords)
 	{
-		int32 NumVertices = InRawMesh.VertexPositions.Num();
-		int32 NumWedges = InRawMesh.WedgeIndices.Num();
+		TVertexAttributesConstRef<FVector> VertexPositions = SrcRawMesh.VertexAttributes().GetAttributesRef<FVector>(MeshAttribute::Vertex::Position);
+		TVertexInstanceAttributesConstRef<FVector> VertexInstanceNormals = SrcRawMesh.VertexInstanceAttributes().GetAttributesRef<FVector>(MeshAttribute::VertexInstance::Normal);
+		TVertexInstanceAttributesConstRef<FVector> VertexInstanceTangents = SrcRawMesh.VertexInstanceAttributes().GetAttributesRef<FVector>(MeshAttribute::VertexInstance::Tangent);
+		TVertexInstanceAttributesConstRef<float> VertexInstanceBinormalSigns = SrcRawMesh.VertexInstanceAttributes().GetAttributesRef<float>(MeshAttribute::VertexInstance::BinormalSign);
+		TVertexInstanceAttributesConstRef<FVector4> VertexInstanceColors = SrcRawMesh.VertexInstanceAttributes().GetAttributesRef<FVector4>(MeshAttribute::VertexInstance::Color);
+		TVertexInstanceAttributesConstRef<FVector2D> VertexInstanceUVs = SrcRawMesh.VertexInstanceAttributes().GetAttributesRef<FVector2D>(MeshAttribute::VertexInstance::TextureCoordinate);
+
+		int32 NumVertices = SrcRawMesh.Vertices().Num();
+		int32 NumWedges = 0;
+		for (const FPolygonID& PolygonID : SrcRawMesh.Polygons().GetElementIDs())
+		{
+			const FMeshPolygon& Polygon = SrcRawMesh.GetPolygon(PolygonID);
+			NumWedges += Polygon.Triangles.Num() * 3;
+		}
 		int32 NumTris = NumWedges / 3;
 
 		if (NumWedges == 0)
 		{
-			return NULL;
+			return nullptr;
 		}
 
 		//assuming everything is left-handed so no need to change winding order and handedness. SSF supports both
@@ -1232,14 +1268,19 @@ bool ZipContentsForUpload(FString InputDirectoryPath, FString OutputFileName)
 		//setup vertex coordinates
 		ssf::ssfList<ssf::ssfVector3> & SsfCoorinates = SgMeshData->Coordinates.Create();
 		SsfCoorinates.Items.resize(NumVertices);
-		for (int32 VertexIndex = 0; VertexIndex < NumVertices; ++VertexIndex)
+		TMap<FVertexID, int32> MeshToSsfVertexID;
+		MeshToSsfVertexID.Reserve(NumVertices);
+		int32 VertexIndex = 0;
+		for(const FVertexID& VertexID : SrcRawMesh.Vertices().GetElementIDs())
 		{
 			ssf::ssfVector3 CurrentVertex;
-			FVector4 Position = GetConversionMatrixYUP().TransformPosition(InRawMesh.VertexPositions[VertexIndex]);
+			FVector4 Position = GetConversionMatrixYUP().TransformPosition(VertexPositions[VertexID]);
 			CurrentVertex.V[0] = double(Position.X);
 			CurrentVertex.V[1] = double(Position.Y);
 			CurrentVertex.V[2] = double(Position.Z);
 			SsfCoorinates.Items[VertexIndex] = CurrentVertex;
+			MeshToSsfVertexID.Add(VertexID, VertexIndex);
+			VertexIndex++;
 		}
 
 		//setup triangle data
@@ -1251,155 +1292,162 @@ bool ZipContentsForUpload(FString InputDirectoryPath, FString OutputFileName)
 		SsfMaterialIndices.Items.resize(NumTris);
 		SsfSmoothingGroups.Items.resize(NumTris);
 
-
-		//Reverse winding switches
-		bool bReverseWinding = true;
-
-		for (int32 TriIndex = 0; TriIndex < NumTris; ++TriIndex)
+		bool bHasNormals = VertexInstanceNormals.GetNumElements() == NumWedges;
+		bool bHasTangents = bHasNormals && (VertexInstanceTangents.GetNumElements() == NumWedges) && (VertexInstanceBinormalSigns.GetNumElements() == NumWedges);
+		
+		ssf::ssfList<ssf::ssfVector3> EmptyList;
+		ssf::ssfList<ssf::ssfVector3>& SsfTangents = bHasTangents ? SgMeshData->Tangents.Create() : EmptyList;
+		ssf::ssfList<ssf::ssfVector3>& SsfBitangents = bHasTangents ? SgMeshData->Bitangents.Create() : EmptyList;
+		ssf::ssfList<ssf::ssfVector3>& SsfNormals = bHasNormals ? SgMeshData->Normals.Create() : EmptyList;
+		if (bHasNormals)
 		{
-			for (int32 CornerIndex = 0; CornerIndex < 3; ++CornerIndex)
+			if (bHasTangents)
 			{
-				int32 DestCornerIndex = bReverseWinding ? 2 - CornerIndex : CornerIndex;
-				SsfTriangleIndices.Items[TriIndex].V[DestCornerIndex] = InRawMesh.WedgeIndices[TriIndex * 3 + CornerIndex];
+				SsfTangents.Items.resize(NumWedges);
+				SsfBitangents.Items.resize(NumWedges);
 			}
+			SsfNormals.Items.resize(NumWedges);
 		}
 
-		for (int32 TriIndex = 0; TriIndex < NumTris; ++TriIndex)
+		const int32 TexCoordNumber = FMath::Min(VertexInstanceUVs.GetNumIndices(), (int32)MAX_MESH_TEXTURE_COORDS);
+		ssf::ssfNamedList<ssf::ssfVector2> SsfTextureCoordinates[MAX_MESH_TEXTURE_COORDS];
+		for (int32 TexCoordIndex = 0; TexCoordIndex < TexCoordNumber; ++TexCoordIndex)
 		{
-			SsfMaterialIndices.Items[TriIndex] = InRawMesh.FaceMaterialIndices[TriIndex];
-			SsfSmoothingGroups.Items[TriIndex] = InRawMesh.FaceSmoothingMasks[TriIndex];
+			//Since SSF uses Named Channels
+			SsfTextureCoordinates[TexCoordIndex].Name = FSimplygonSSFHelper::TCHARToSSFString(*FString::Printf(TEXT("TexCoord%d"), TexCoordIndex));
+			SsfTextureCoordinates[TexCoordIndex].Items.resize(NumWedges);
 		}
 
-		SgMeshData->MaterialIndices.Create();
-
-		//setup texcoords
-		for (int32 TexCoordIndex = 0; TexCoordIndex < MAX_MESH_TEXTURE_COORDS; ++TexCoordIndex)
+		ssf::ssfNamedList<ssf::ssfVector4> SsfColorMap;
+		bool bHasVertexColor = VertexInstanceColors.GetNumElements() == NumWedges;
+		if (bHasVertexColor)
 		{
-			const TArray<FVector2D>& SrcTexCoords = (TexCoordIndex == 0 && InTexCoords.Num() == NumWedges) ? InTexCoords : InRawMesh.WedgeTexCoords[TexCoordIndex];
+			//setup the color named channel . Currently its se to index zero. If multiple colors channel are need then use an index instead of 0
+			SsfColorMap.Name = FSimplygonSSFHelper::TCHARToSSFString(*FString::Printf(TEXT("Colors%d"), 0));
+			SsfColorMap.Items.resize(NumWedges);
+		}
 
-			if (SrcTexCoords.Num() == NumWedges)
-			{
-				ssf::ssfNamedList<ssf::ssfVector2> SsfTextureCoordinates;
+		//Smooth group
+		TArray<uint32> FaceSmoothingMasks;
+		FMeshDescriptionOperations::ConvertHardEdgesToSmoothGroup(SrcRawMesh, FaceSmoothingMasks);
 
-				//Since SSF uses Named Channels
-				SsfTextureCoordinates.Name = FSimplygonSSFHelper::TCHARToSSFString(*FString::Printf(TEXT("TexCoord%d"), TexCoordIndex));
-				SsfTextureCoordinates.Items.resize(NumWedges);
+		//Reverse winding switches
+		bool bReverseWinding = true;
+		
+		int32 TriangleIndex = 0;
+		for (const FPolygonID& PolygonID : SrcRawMesh.Polygons().GetElementIDs())
+		{
+			const FMeshPolygon& Polygon = SrcRawMesh.GetPolygon(PolygonID);
+			
+			FPolygonGroupID PolygonGroupID = SrcRawMesh.GetPolygonPolygonGroup(PolygonID);
+			int32 MaterialIndex = PolygonGroupID.GetValue();
 
-				int32 WedgeIndex = 0;
-				for (int32 TriIndex = 0; TriIndex < NumTris; ++TriIndex)
+			for (const FMeshTriangle& Triangle : Polygon.Triangles)
+			{
+				for (int32 CornerIndex = 0; CornerIndex < 3; ++CornerIndex)
 				{
-					int32 MaterialIndex = InRawMesh.FaceMaterialIndices[TriIndex];
-					// Compute texture bounds for current material.
-					float MinU = 0, ScaleU = 1;
-					float MinV = 0, ScaleV = 1;
-
-					if (InTextureBounds.IsValidIndex(MaterialIndex) && TexCoordIndex == 0 && InTexCoords.Num() == 0)
+					int32 DestCornerIndex = bReverseWinding ? 2 - CornerIndex : CornerIndex;
+					FVertexInstanceID VertexInstanceID = Triangle.GetVertexInstanceID(CornerIndex);
+					SsfTriangleIndices.Items[TriangleIndex].V[DestCornerIndex] = MeshToSsfVertexID[SrcRawMesh.GetVertexInstanceVertex(VertexInstanceID)];
+					
+					//NTBs
+					if (bHasNormals)
 					{
-						const FBox2D& Bounds = InTextureBounds[MaterialIndex];
-						if (Bounds.GetArea() > 0)
+						FVector Normal = VertexInstanceNormals[VertexInstanceID];
+						if (bHasTangents)
 						{
-							MinU = Bounds.Min.X;
-							MinV = Bounds.Min.Y;
-							ScaleU = 1.0f / (Bounds.Max.X - Bounds.Min.X);
-							ScaleV = 1.0f / (Bounds.Max.Y - Bounds.Min.Y);
+							FVector Tangent = VertexInstanceTangents[VertexInstanceID];
+							FVector Bitangent = FVector::CrossProduct(Normal, Tangent).GetSafeNormal() * VertexInstanceBinormalSigns[VertexInstanceID];
+
+							ssf::ssfVector3 SsfTangent;
+							FVector4 Tangent4 = GetConversionMatrixYUP().TransformPosition(Tangent);
+							SsfTangent.V[0] = double(Tangent4.X);
+							SsfTangent.V[1] = double(Tangent4.Y);
+							SsfTangent.V[2] = double(Tangent4.Z);
+							SsfTangents.Items[TriangleIndex * 3 + DestCornerIndex] = SsfTangent;
+
+							ssf::ssfVector3 SsfBitangent;
+							FVector4 Bitangent4 = GetConversionMatrixYUP().TransformPosition(Bitangent);
+							SsfBitangent.V[0] = double(Bitangent4.X);
+							SsfBitangent.V[1] = double(Bitangent4.Y);
+							SsfBitangent.V[2] = double(Bitangent4.Z);
+							SsfBitangents.Items[TriangleIndex * 3 + DestCornerIndex] = SsfBitangent;
 						}
+
+						ssf::ssfVector3 SsfNormal;
+						FVector4 Normal4 = GetConversionMatrixYUP().TransformPosition(Normal);
+						SsfNormal.V[0] = double(Normal4.X);
+						SsfNormal.V[1] = double(Normal4.Y);
+						SsfNormal.V[2] = double(Normal4.Z);
+						SsfNormals.Items[TriangleIndex * 3 + DestCornerIndex] = SsfNormal;
 					}
 
-					for (int32 CornerIndex = 0; CornerIndex < 3; ++CornerIndex)
+					//Vertex color
+					if (bHasVertexColor)
 					{
-						const FVector2D& TexCoord = SrcTexCoords[TriIndex * 3 + CornerIndex];
-						ssf::ssfVector2 temp;
-						temp.V[0] = (TexCoord.X - MinU) * ScaleU;
-						temp.V[1] = (TexCoord.Y - MinV) * ScaleV;
-						int32 DestCornerIndex = bReverseWinding ? 2 - CornerIndex : CornerIndex;
-						SsfTextureCoordinates.Items[TriIndex * 3 + DestCornerIndex] = temp;
+						FLinearColor LinearColor = VertexInstanceColors[VertexInstanceID];
+						SsfColorMap.Items[TriangleIndex * 3 + DestCornerIndex].V[0] = LinearColor.R;
+						SsfColorMap.Items[TriangleIndex * 3 + DestCornerIndex].V[1] = LinearColor.G;
+						SsfColorMap.Items[TriangleIndex * 3 + DestCornerIndex].V[2] = LinearColor.B;
+						SsfColorMap.Items[TriangleIndex * 3 + DestCornerIndex].V[3] = LinearColor.A;
 					}
-				}
-
-				SgMeshData->TextureCoordinatesList.push_back(SsfTextureCoordinates);
-			}
-		}
-
-		//setup colors
-		if (InRawMesh.WedgeColors.Num() == NumWedges)
-		{
-			//setup the color named channel . Currently its se to index zero. If multiple colors channel are need then use an index instead of 0
-			ssf::ssfNamedList<ssf::ssfVector4> SsfColorMap;
-			SsfColorMap.Name = FSimplygonSSFHelper::TCHARToSSFString(*FString::Printf(TEXT("Colors%d"), 0));
-			SsfColorMap.Items.resize(NumWedges);
-			for (int32 TriIndex = 0; TriIndex < NumTris; ++TriIndex)
-			{
-				for (int32 CornerIndex = 0; CornerIndex < 3; ++CornerIndex)
-				{
-					int32 DestCornerIndex = bReverseWinding ? 2 - CornerIndex : CornerIndex;
-					FLinearColor LinearColor(InRawMesh.WedgeColors[TriIndex * 3 + CornerIndex]);
-					SsfColorMap.Items[TriIndex * 3 + DestCornerIndex].V[0] = LinearColor.R;
-					SsfColorMap.Items[TriIndex * 3 + DestCornerIndex].V[1] = LinearColor.G;
-					SsfColorMap.Items[TriIndex * 3 + DestCornerIndex].V[2] = LinearColor.B;
-					SsfColorMap.Items[TriIndex * 3 + DestCornerIndex].V[3] = LinearColor.A;
-				}
-			}
-			SgMeshData->ColorsList.push_back(SsfColorMap);
-		}
-		 
-		if (InRawMesh.WedgeTangentZ.Num() == NumWedges)
-		{
-			if (InRawMesh.WedgeTangentX.Num() == NumWedges && InRawMesh.WedgeTangentY.Num() == NumWedges)
-			{
-				ssf::ssfList<ssf::ssfVector3> & SsfTangents = SgMeshData->Tangents.Create();
-				SsfTangents.Items.resize(NumWedges);
+					
 
-				for (int32 TriIndex = 0; TriIndex < NumTris; ++TriIndex)
-				{
-					for (int32 CornerIndex = 0; CornerIndex < 3; ++CornerIndex)
+					//Texcoords
+					for (int32 TexCoordIndex = 0; TexCoordIndex < TexCoordNumber; ++TexCoordIndex)
 					{
-						int32 DestCornerIndex = bReverseWinding ? 2 - CornerIndex : CornerIndex;
-						ssf::ssfVector3 SsfTangent;
-						FVector4 Tangent = GetConversionMatrixYUP().TransformPosition(InRawMesh.WedgeTangentX[TriIndex * 3 + CornerIndex]);
-						SsfTangent.V[0] = double(Tangent.X);
-						SsfTangent.V[1] = double(Tangent.Y);
-						SsfTangent.V[2] = double(Tangent.Z);
-						SsfTangents.Items[TriIndex * 3 + DestCornerIndex] = SsfTangent;
+						bool bUseInTexCoord = (TexCoordIndex == 0 && InTexCoords.Num() == NumWedges);
+						int32 NumTexCoord = bUseInTexCoord ? InTexCoords.Num() : VertexInstanceUVs.GetNumElements();
+						if (NumTexCoord == NumWedges)
+						{
+							// Compute texture bounds for current material.
+							float MinU = 0, ScaleU = 1;
+							float MinV = 0, ScaleV = 1;
+							if (InTextureBounds.IsValidIndex(MaterialIndex) && TexCoordIndex == 0 && InTexCoords.Num() == 0)
+							{
+								const FBox2D& Bounds = InTextureBounds[MaterialIndex];
+								if (Bounds.GetArea() > 0)
+								{
+									MinU = Bounds.Min.X;
+									MinV = Bounds.Min.Y;
+									ScaleU = 1.0f / (Bounds.Max.X - Bounds.Min.X);
+									ScaleV = 1.0f / (Bounds.Max.Y - Bounds.Min.Y);
+								}
+							}
+
+							const FVector2D& TexCoord = bUseInTexCoord ? InTexCoords[TriangleIndex * 3 + CornerIndex] : VertexInstanceUVs.Get(VertexInstanceID, TexCoordIndex);
+							ssf::ssfVector2 temp;
+							temp.V[0] = (TexCoord.X - MinU) * ScaleU;
+							temp.V[1] = (TexCoord.Y - MinV) * ScaleV;
+							SsfTextureCoordinates[TexCoordIndex].Items[TriangleIndex * 3 + DestCornerIndex] = temp;
+						}
 					}
 
 				}
-				 
+				
+				//Material
+				SsfMaterialIndices.Items[TriangleIndex] = MaterialIndex;
 
-				ssf::ssfList<ssf::ssfVector3> & SsfBitangents = SgMeshData->Bitangents.Create();
-				SsfBitangents.Items.resize(NumWedges);
-				for (int32 TriIndex = 0; TriIndex < NumTris; ++TriIndex)
-				{
-					for (int32 CornerIndex = 0; CornerIndex < 3; ++CornerIndex)
-					{
-						int32 DestCornerIndex = bReverseWinding ? 2 - CornerIndex : CornerIndex;
-						ssf::ssfVector3 SsfBitangent;
-						FVector4 Bitangent = GetConversionMatrixYUP().TransformPosition(InRawMesh.WedgeTangentY[TriIndex * 3 + CornerIndex]);
-						SsfBitangent.V[0] = double(Bitangent.X);
-						SsfBitangent.V[1] = double(Bitangent.Y);
-						SsfBitangent.V[2] = double(Bitangent.Z);
-						SsfBitangents.Items[TriIndex * 3 + DestCornerIndex] = SsfBitangent;
-					}
-				}
-				 
+				//Smooth group
+				SsfSmoothingGroups.Items[TriangleIndex] = FaceSmoothingMasks[TriangleIndex];
+
+				TriangleIndex++;
 			}
+		}
 
-			ssf::ssfList<ssf::ssfVector3> & SsfNormals = SgMeshData->Normals.Create();
-			SsfNormals.Items.resize(NumWedges);
+		SgMeshData->MaterialIndices.Create();
 
-			for (int32 TriIndex = 0; TriIndex < NumTris; ++TriIndex)
-			{
-				for (int32 CornerIndex = 0; CornerIndex < 3; ++CornerIndex)
-				{
-					int32 DestCornerIndex = bReverseWinding ? 2 - CornerIndex : CornerIndex;
-					ssf::ssfVector3 SsfNormal;
-					FVector4 Normal = GetConversionMatrixYUP().TransformPosition(InRawMesh.WedgeTangentZ[TriIndex * 3 + CornerIndex]);
-					SsfNormal.V[0] = double(Normal.X);
-					SsfNormal.V[1] = double(Normal.Y);
-					SsfNormal.V[2] = double(Normal.Z);
-					SsfNormals.Items[TriIndex * 3 + DestCornerIndex] = SsfNormal;
-				}
-			}
+		//Push back all the data...
+
+		for (int32 TexCoordIndex = 0; TexCoordIndex < TexCoordNumber; ++TexCoordIndex)
+		{
+			SgMeshData->TextureCoordinatesList.push_back(SsfTextureCoordinates[TexCoordIndex]);
 		}
+		if (bHasVertexColor)
+		{
+			SgMeshData->ColorsList.push_back(SsfColorMap);
+		}
+
 
 		return SgMeshData;
 	}
diff --git a/Engine/Source/Developer/SimplygonSwarm/SimplygonSwarm.Build.cs b/Engine/Source/Developer/SimplygonSwarm/SimplygonSwarm.Build.cs
--- a/Engine/Source/Developer/SimplygonSwarm/SimplygonSwarm.Build.cs
+++ b/Engine/Source/Developer/SimplygonSwarm/SimplygonSwarm.Build.cs
@@ -26,13 +26,15 @@ public class SimplygonSwarm : ModuleRules
 				"CoreUObject",
 				"Engine",
 				"RenderCore",
-				"RawMesh",
+                "MeshDescription",
+                "MeshDescriptionOperations",
                 "MeshBoneReduction",
                 "ImageWrapper",
                 "HTTP",
                 "Json",
                 "UnrealEd",
-                "MaterialUtilities"
+                "MaterialUtilities",
+                "MeshMergeUtilities"
 			}
         );
 
