Use explicit FVector4 -> FLinearColor constructor to fix Clang compile error.

#rnx

diff --git a/Engine/Source/Developer/SimplygonSwarm/Private/SimplygonSwarm.cpp b/Engine/Source/Developer/SimplygonSwarm/Private/SimplygonSwarm.cpp
--- a/Engine/Source/Developer/SimplygonSwarm/Private/SimplygonSwarm.cpp
+++ b/Engine/Source/Developer/SimplygonSwarm/Private/SimplygonSwarm.cpp
@@ -1375,7 +1375,7 @@ bool ZipContentsForUpload(FString InputDirectoryPath, FString OutputFileName)
 					//Vertex color
 					if (bHasVertexColor)
 					{
-						FLinearColor LinearColor = VertexInstanceColors[VertexInstanceID];
+						FLinearColor LinearColor(VertexInstanceColors[VertexInstanceID]);
 						SsfColorMap.Items[TriangleIndex * 3 + DestCornerIndex].V[0] = LinearColor.R;
 						SsfColorMap.Items[TriangleIndex * 3 + DestCornerIndex].V[1] = LinearColor.G;
 						SsfColorMap.Items[TriangleIndex * 3 + DestCornerIndex].V[2] = LinearColor.B;
